Move ControlRigPhysics out of PhysicsControl and into its own plugin. No change in functionality etc.

diff --git a/Engine/Plugins/Experimental/ControlRigPhysics/Source/ControlRigPhysics/Public/RigPhysicsBodyComponent.h b/Engine/Plugins/Experimental/ControlRigPhysics/Source/ControlRigPhysics/Public/RigPhysicsBodyComponent.h
new file mode 100644
--- /dev/null
+++ b/Engine/Plugins/Experimental/ControlRigPhysics/Source/ControlRigPhysics/Public/RigPhysicsBodyComponent.h
@@ -0,0 +1,75 @@
+// Copyright Epic Games, Inc. All Rights Reserved.
+
+#pragma once
+
+#include "RigPhysicsData.h"
+
+#include "PhysicsControlData.h"
+
+#include "Rigs/RigPhysics.h"
+#include "Rigs/RigHierarchyDefines.h"
+
+#include "RigPhysicsBodyComponent.generated.h"
+
+#define UE_API CONTROLRIGPHYSICS_API
+
+/**
+ * A component that can be added to a joint/element that defines how a physical body can be "attached" to it.
+ * The body supports dynamic movement, collision, and a physics joint with this body's parent in the hierarchy.
+ */
+USTRUCT(BlueprintType)
+struct FRigPhysicsBodyComponent : public FRigBaseComponent
+{
+public:
+
+	GENERATED_BODY()
+	DECLARE_RIG_COMPONENT_METHODS(FRigPhysicsBodyComponent)
+
+	UPROPERTY(BlueprintReadOnly, EditAnywhere, Category = Physics, meta = (ShowOnlyInnerProperties))
+	FRigPhysicsBodySolverSettings BodySolverSettings;
+
+	UPROPERTY(BlueprintReadOnly, EditAnywhere, Category = Physics, meta = (ShowOnlyInnerProperties))
+	FRigPhysicsDynamics Dynamics;
+
+	UPROPERTY(BlueprintReadOnly, EditAnywhere, Category = Physics, meta = (ShowOnlyInnerProperties))
+	FRigPhysicsCollision Collision;
+
+	UPROPERTY(BlueprintReadOnly, EditAnywhere, Category = Physics)
+	FPhysicsControlModifierData BodyData;
+
+	// The target for when this body is kinematic
+	UPROPERTY(BlueprintReadOnly, EditAnywhere, Category = Physics)
+	FTransform KinematicTarget;
+
+	UPROPERTY(BlueprintReadOnly, EditAnywhere, Category = Physics)
+	ERigPhysicsKinematicTargetSpace KinematicTargetSpace = ERigPhysicsKinematicTargetSpace::OffsetInBoneSpace;
+
+	// A list of body components with which we should not collide. The solver component can also be included.
+	UPROPERTY(BlueprintReadOnly, EditAnywhere, Category = Physics)
+	TArray<FRigComponentKey> NoCollisionBodies;
+
+	// Removes any existing collision, and replaces it with a shape calculated from the joint
+	// positions (if possible). The shape will be a single box.
+	// @param MinAspectRatio the minimum box extent, as a proportion of the maximum box extent.
+	UE_API void AutoCalculateCollision(URigHierarchy* InHierarchy, float MinAspectRatio = 0.25f, float MinSize = 0.0f);
+
+	UE_API virtual void Save(FArchive& Ar) override;
+	UE_API virtual void Load(FArchive& Ar) override;
+
+#if WITH_EDITOR
+	UE_API virtual const FSlateIcon& GetIconForUI() const override;
+#endif
+
+	UE_API virtual bool CanBeAddedTo(
+		const FRigElementKey& InElementKey, const URigHierarchy* InHierarchy, FString* OutFailureReason) const override;
+
+	UE_API virtual void OnAddedToHierarchy(URigHierarchy* InHierarchy, URigHierarchyController* InController) override;
+
+	UE_API virtual void OnRigHierarchyKeyChanged(const FRigHierarchyKey& InOldKey, const FRigHierarchyKey& InNewKey) override;
+
+	virtual FName GetDefaultComponentName() const override { return GetDefaultName(); }
+
+	static FName GetDefaultName() { return TEXT("PhysicsBody"); }
+};
+
+#undef UE_API
