Incremental improvements to Control Rig Physics. Allow body components to be owned by any element (e.g. a control, not just a bone). Fix visualization of joint limits. Rename "Add" to "Spawn". Remove "Component" from the displayed function names. Remove any dependency on ControlRig from the simulation itself (it's still there in the rig units).

diff --git a/Engine/Plugins/Experimental/ControlRigPhysics/Source/ControlRigPhysics/Private/RigPhysicsBodyComponent.cpp b/Engine/Plugins/Experimental/ControlRigPhysics/Source/ControlRigPhysics/Private/RigPhysicsBodyComponent.cpp
--- a/Engine/Plugins/Experimental/ControlRigPhysics/Source/ControlRigPhysics/Private/RigPhysicsBodyComponent.cpp
+++ b/Engine/Plugins/Experimental/ControlRigPhysics/Source/ControlRigPhysics/Private/RigPhysicsBodyComponent.cpp
@@ -93,21 +93,6 @@ FControlRigPhysicsEditorStyle::Get().GetStyleSetName(), "ControlRigPhysics.Compo
 }
 #endif
 
-//======================================================================================================================
-bool FRigPhysicsBodyComponent::CanBeAddedTo(
-	const FRigElementKey& InElementKey, const URigHierarchy* InHierarchy, FString* OutFailureReason) const
-{
-	if(InElementKey.Type != ERigElementType::Bone)
-	{
-		if(OutFailureReason)
-		{
-			*OutFailureReason = TEXT("Physics components can only be added to bones.");
-		}
-		return false;
-	}
-	return true;
-}
-
 //======================================================================================================================
 void FRigPhysicsBodyComponent::OnAddedToHierarchy(URigHierarchy* InHierarchy, URigHierarchyController* InController)
 {
diff --git a/Engine/Plugins/Experimental/ControlRigPhysics/Source/ControlRigPhysics/Private/RigPhysicsExecution.cpp b/Engine/Plugins/Experimental/ControlRigPhysics/Source/ControlRigPhysics/Private/RigPhysicsExecution.cpp
--- a/Engine/Plugins/Experimental/ControlRigPhysics/Source/ControlRigPhysics/Private/RigPhysicsExecution.cpp
+++ b/Engine/Plugins/Experimental/ControlRigPhysics/Source/ControlRigPhysics/Private/RigPhysicsExecution.cpp
@@ -73,19 +73,16 @@ FRigUnit_TrackInputPose_Execute()
 	{
 		return;
 	}
-	if (UControlRig* ControlRig = Cast<UControlRig>(ExecuteContext.Hierarchy->GetOuter()))
+	if (FRigPhysicsSolverComponent* SolverComponent = Cast<FRigPhysicsSolverComponent>(
+		ExecuteContext.Hierarchy->FindComponent(PhysicsSolverComponentKey)))
 	{
-		if (FRigPhysicsSolverComponent* SolverComponent = Cast<FRigPhysicsSolverComponent>(
-			ExecuteContext.Hierarchy->FindComponent(PhysicsSolverComponentKey)))
+		if (bForceNumberOfFrames)
 		{
-			if (bForceNumberOfFrames)
-			{
-				SolverComponent->TrackInputCounter = NumberOfFrames;
-			}
-			else
-			{
-				SolverComponent->TrackInputCounter = FMath::Max(SolverComponent->TrackInputCounter, NumberOfFrames);
-			}
+			SolverComponent->TrackInputCounter = NumberOfFrames;
+		}
+		else
+		{
+			SolverComponent->TrackInputCounter = FMath::Max(SolverComponent->TrackInputCounter, NumberOfFrames);
 		}
 	}
 }
@@ -102,10 +99,11 @@ FRigUnit_StepPhysicsSolver_Execute()
 	{
 		return;
 	}
-	if (UControlRig* ControlRig = Cast<UControlRig>(ExecuteContext.Hierarchy->GetOuter()))
+	URigHierarchy& Hierarchy = *ExecuteContext.Hierarchy;
+	if (UControlRig* ControlRig = Cast<UControlRig>(Hierarchy.GetOuter()))
 	{
 		if (FRigPhysicsSolverComponent* SolverComponent = Cast<FRigPhysicsSolverComponent>(
-			ExecuteContext.Hierarchy->FindComponent(PhysicsSolverComponentKey)))
+			Hierarchy.FindComponent(PhysicsSolverComponentKey)))
 		{
 			FRigPhysicsSimulation* PhysicsSimulation = Cast<FRigPhysicsSimulation>(
 				ControlRig->GetPhysicsSimulation(PhysicsSolverComponentKey));
@@ -113,11 +111,12 @@ FRigUnit_StepPhysicsSolver_Execute()
 			// Make the solver and instantiate if necessary
 			if (!PhysicsSimulation)
 			{
-				TSharedPtr<FRigPhysicsSimulation> NewPhysicsSimulation = MakeShared<FRigPhysicsSimulation>(ControlRig);
+				TSharedPtr<FRigPhysicsSimulation> NewPhysicsSimulation = MakeShared<FRigPhysicsSimulation>(
+					FName(ControlRig->GetName()));
 
 				if (ControlRig->RegisterPhysicsSimulation(NewPhysicsSimulation, PhysicsSolverComponentKey))
 				{
-					NewPhysicsSimulation->Instantiate(ExecuteContext, SolverComponent);
+					NewPhysicsSimulation->Instantiate(ExecuteContext, Hierarchy, SolverComponent);
 				}
 				PhysicsSimulation = NewPhysicsSimulation.Get();
 			}
@@ -125,7 +124,7 @@ FRigUnit_StepPhysicsSolver_Execute()
 			if (PhysicsSimulation)
 			{
 				PhysicsSimulation->StepSimulation(
-					ExecuteContext, SolverComponent, DeltaTimeOverride, SimulationSpaceDeltaTimeOverride, Alpha);
+					ExecuteContext, Hierarchy, SolverComponent, DeltaTimeOverride, SimulationSpaceDeltaTimeOverride, Alpha);
 				int32 VisualizationOverride = CVarControlRigPhysicsVisualizationOverride.GetValueOnAnyThread();
 				if (VisualizationOverride == 1 || (VisualizationOverride < 0 && VisualizationSettings.bEnableVisualization))
 				{
@@ -146,13 +145,11 @@ FRigUnit_SetPhysicsSolverSimulationSpaceSettings_Execute()
 	{
 		return;
 	}
-	if (UControlRig* ControlRig = Cast<UControlRig>(ExecuteContext.Hierarchy->GetOuter()))
+	URigHierarchy& Hierarchy = *ExecuteContext.Hierarchy;
+	if (FRigPhysicsSolverComponent* SolverComponent = Cast<FRigPhysicsSolverComponent>(
+		Hierarchy.FindComponent(PhysicsSolverComponentKey)))
 	{
-		if (FRigPhysicsSolverComponent* SolverComponent = Cast<FRigPhysicsSolverComponent>(
-			ExecuteContext.Hierarchy->FindComponent(PhysicsSolverComponentKey)))
-		{
-			SolverComponent->SimulationSpaceSettings = SimulationSpaceSettings;
-		}
+		SolverComponent->SimulationSpaceSettings = SimulationSpaceSettings;
 	}
 }
 
@@ -164,15 +161,13 @@ FRigUnit_SetPhysicsSolverExternalVelocity_Execute()
 	{
 		return;
 	}
-	if (UControlRig* ControlRig = Cast<UControlRig>(ExecuteContext.Hierarchy->GetOuter()))
+	URigHierarchy& Hierarchy = *ExecuteContext.Hierarchy;
+	if (FRigPhysicsSolverComponent* SolverComponent = Cast<FRigPhysicsSolverComponent>(
+		Hierarchy.FindComponent(PhysicsSolverComponentKey)))
 	{
-		if (FRigPhysicsSolverComponent* SolverComponent = Cast<FRigPhysicsSolverComponent>(
-			ExecuteContext.Hierarchy->FindComponent(PhysicsSolverComponentKey)))
-		{
-			SolverComponent->SimulationSpaceSettings.ExternalLinearVelocity = ExternalLinearVelocity;
-			SolverComponent->SimulationSpaceSettings.ExternalAngularVelocity = ExternalAngularVelocity;
-			SolverComponent->SimulationSpaceSettings.ExternalTurbulenceVelocity = ExternalTurbulenceVelocity;
-		}
+		SolverComponent->SimulationSpaceSettings.ExternalLinearVelocity = ExternalLinearVelocity;
+		SolverComponent->SimulationSpaceSettings.ExternalAngularVelocity = ExternalAngularVelocity;
+		SolverComponent->SimulationSpaceSettings.ExternalTurbulenceVelocity = ExternalTurbulenceVelocity;
 	}
 }
 
@@ -184,16 +179,18 @@ FRigUnit_InstantiatePhysics_Execute()
 	{
 		return;
 	}
-	if (UControlRig* ControlRig = Cast<UControlRig>(ExecuteContext.Hierarchy->GetOuter()))
+	URigHierarchy& Hierarchy = *ExecuteContext.Hierarchy;
+	if (UControlRig* ControlRig = Cast<UControlRig>(Hierarchy.GetOuter()))
 	{
 		if (FRigPhysicsSolverComponent* SolverComponent = Cast<FRigPhysicsSolverComponent>(
 			ExecuteContext.Hierarchy->FindComponent(PhysicsSolverComponentKey)))
 		{
-			TSharedPtr<FRigPhysicsSimulation> PhysicsSimulation = MakeShared<FRigPhysicsSimulation>(ControlRig);
+			TSharedPtr<FRigPhysicsSimulation> PhysicsSimulation = MakeShared<FRigPhysicsSimulation>(
+				FName(ControlRig->GetName()));
 
 			if (ControlRig->RegisterPhysicsSimulation(PhysicsSimulation, PhysicsSolverComponentKey))
 			{
-				PhysicsSimulation->Instantiate(ExecuteContext, SolverComponent);
+				PhysicsSimulation->Instantiate(ExecuteContext, Hierarchy, SolverComponent);
 			}
 		}
 	}
diff --git a/Engine/Plugins/Experimental/ControlRigPhysics/Source/ControlRigPhysics/Private/RigPhysicsSimulation.cpp b/Engine/Plugins/Experimental/ControlRigPhysics/Source/ControlRigPhysics/Private/RigPhysicsSimulation.cpp
--- a/Engine/Plugins/Experimental/ControlRigPhysics/Source/ControlRigPhysics/Private/RigPhysicsSimulation.cpp
+++ b/Engine/Plugins/Experimental/ControlRigPhysics/Source/ControlRigPhysics/Private/RigPhysicsSimulation.cpp
@@ -10,7 +10,6 @@
 #include "PhysicsControlPoseData.h"
 #include "PhysicsControlHelpers.h"
 
-#include "ControlRig.h"
 #include "Rigs/RigHierarchy.h"
 #include "Rigs/RigHierarchyController.h"
 #include "Units/RigUnitContext.h"
@@ -65,22 +64,20 @@ constexpr int32 ConstraintChildIndex = 0;
 constexpr int32 ConstraintParentIndex = 1;
 
 //======================================================================================================================
-FRigPhysicsSimulation::FRigPhysicsSimulation(TObjectPtr<UControlRig> InOwningControlRig)
-	: FRigPhysicsSimulationBase(FRigPhysicsSimulation::StaticStruct()), OwningControlRig(InOwningControlRig)
+FRigPhysicsSimulation::FRigPhysicsSimulation(const FName InOwnerName)
+	: FRigPhysicsSimulationBase(FRigPhysicsSimulation::StaticStruct()), OwnerName(InOwnerName)
 {
 }
 
 //======================================================================================================================
 bool FRigPhysicsSimulation::ShouldComponentBeInSimulation(
-	FRigComponentKey SolverComponentKey, FRigComponentKey ComponentKey) const
+	const URigHierarchy& Hierarchy, FRigComponentKey SolverComponentKey, FRigComponentKey ComponentKey) const
 {
-	const URigHierarchy* Hierarchy = OwningControlRig->GetHierarchy();
-
 	const FRigPhysicsBodyComponent* PhysicsComponent = Cast<FRigPhysicsBodyComponent>(
-		Hierarchy->FindComponent(ComponentKey));
+		Hierarchy.FindComponent(ComponentKey));
 
 	const FRigPhysicsSolverComponent* PhysicsSolverComponent = Cast<FRigPhysicsSolverComponent>(
-		Hierarchy->FindComponent(SolverComponentKey));
+		Hierarchy.FindComponent(SolverComponentKey));
 
 	if (!PhysicsComponent || !PhysicsSolverComponent)
 	{
@@ -102,7 +99,7 @@ bool FRigPhysicsSimulation::ShouldComponentBeInSimulation(
 		FRigElementKey ElementKey = PhysicsComponent->GetElementKey();
 		while (ElementKey.IsValid())
 		{
-			for (FRigComponentKey CK : Hierarchy->GetComponentKeys(ElementKey))
+			for (FRigComponentKey CK : Hierarchy.GetComponentKeys(ElementKey))
 			{
 				if (CK == SolverComponentKey)
 				{
@@ -110,7 +107,7 @@ bool FRigPhysicsSimulation::ShouldComponentBeInSimulation(
 				}
 			}
 			// Note that getting the parent of an element at the root doesn't return the top-level element
-			ElementKey = Hierarchy->GetFirstParent(ElementKey);
+			ElementKey = Hierarchy.GetFirstParent(ElementKey);
 		}
 	}
 
@@ -126,8 +123,7 @@ void FRigPhysicsSimulation::InitialiseSimulation(const FRigPhysicsSolverComponen
 	Simulation = MakeShared<ImmediatePhysics::FSimulation>();
 
 #if CHAOS_DEBUG_NAME
-	FString SimName = TEXT("ControlRigPhysics-");
-	SimName += OwningControlRig->GetName();
+	FString SimName = TEXT("ControlRigPhysics-") + OwnerName.ToString();
 	Simulation->SetDebugName(FName(SimName));
 #endif
 
@@ -148,24 +144,24 @@ void FRigPhysicsSimulation::InitialiseSimulation(const FRigPhysicsSolverComponen
 }
 
 //======================================================================================================================
-void FRigPhysicsSimulation::InitialiseControlRecords(const FRigPhysicsSolverComponent* SolverComponent)
+void FRigPhysicsSimulation::InitialiseControlRecords(
+	const URigHierarchy& Hierarchy, const FRigPhysicsSolverComponent* SolverComponent)
 {
 	ensure(ControlRecords.IsEmpty());
 
-	const URigHierarchy* Hierarchy = OwningControlRig->GetHierarchy();
-	if (!Hierarchy || !SolverComponent)
+	if (!SolverComponent)
 	{
 		return;
 	}
 
 	FRigComponentKey SolverComponentKey = SolverComponent->GetKey();
 
-	TArray<FRigComponentKey> AllComponentKeys = Hierarchy->GetAllComponentKeys();
+	TArray<FRigComponentKey> AllComponentKeys = Hierarchy.GetAllComponentKeys();
 
 	for (FRigComponentKey ComponentKey : AllComponentKeys)
 	{
 		if (const FRigPhysicsControlComponent* ControlComponent = Cast<FRigPhysicsControlComponent>(
-			Hierarchy->FindComponent(ComponentKey)))
+			Hierarchy.FindComponent(ComponentKey)))
 		{
 			// The authored body components may be blank, in which case we need to find the automatic ones.
 			FRigControlRecord ControlRecord;
@@ -175,10 +171,10 @@ void FRigPhysicsSimulation::InitialiseControlRecords(const FRigPhysicsSolverComp
 			// Automate the child
 			if (!ControlRecord.ChildBodyComponentKey.IsValid())
 			{
-				TArray<FRigComponentKey> SiblingComponentKeys = Hierarchy->GetComponentKeys(ComponentKey.ElementKey);
+				TArray<FRigComponentKey> SiblingComponentKeys = Hierarchy.GetComponentKeys(ComponentKey.ElementKey);
 				for (FRigComponentKey SiblingComponentKey : SiblingComponentKeys)
 				{
-					if (ShouldComponentBeInSimulation(SolverComponentKey, SiblingComponentKey))
+					if (ShouldComponentBeInSimulation(Hierarchy, SolverComponentKey, SiblingComponentKey))
 					{
 						ControlRecord.ChildBodyComponentKey = SiblingComponentKey;
 						break;
@@ -190,11 +186,11 @@ void FRigPhysicsSimulation::InitialiseControlRecords(const FRigPhysicsSolverComp
 			{
 				if (!ControlRecord.ParentBodyComponentKey.IsValid())
 				{
-					FRigElementKey ParentElementKey = Hierarchy->GetFirstParent(ComponentKey.ElementKey);
-					TArray<FRigComponentKey> ParentComponentKeys = Hierarchy->GetComponentKeys(ParentElementKey);
+					FRigElementKey ParentElementKey = Hierarchy.GetFirstParent(ComponentKey.ElementKey);
+					TArray<FRigComponentKey> ParentComponentKeys = Hierarchy.GetComponentKeys(ParentElementKey);
 					for (FRigComponentKey ParentComponentKey : ParentComponentKeys)
 					{
-						if (ShouldComponentBeInSimulation(SolverComponentKey, ParentComponentKey))
+						if (ShouldComponentBeInSimulation(Hierarchy, SolverComponentKey, ParentComponentKey))
 						{
 							ControlRecord.ParentBodyComponentKey = ParentComponentKey;
 							break;
@@ -203,11 +199,11 @@ void FRigPhysicsSimulation::InitialiseControlRecords(const FRigPhysicsSolverComp
 				}
 			}
 
-			if (ShouldComponentBeInSimulation(SolverComponentKey, ControlRecord.ChildBodyComponentKey))
+			if (ShouldComponentBeInSimulation(Hierarchy, SolverComponentKey, ControlRecord.ChildBodyComponentKey))
 			{
 				// Here, an invalid parent component key indicates a sim-space control. 
 				if (!ControlRecord.ParentBodyComponentKey.IsValid() ||
-					ShouldComponentBeInSimulation(SolverComponentKey, ControlRecord.ParentBodyComponentKey))
+					ShouldComponentBeInSimulation(Hierarchy, SolverComponentKey, ControlRecord.ParentBodyComponentKey))
 				{
 					// Just make the record for now - it will be instantiated later
 					ControlRecords.Add(ComponentKey, ControlRecord);
@@ -218,24 +214,23 @@ void FRigPhysicsSimulation::InitialiseControlRecords(const FRigPhysicsSolverComp
 }
 
 //======================================================================================================================
-void FRigPhysicsSimulation::InitialiseJointRecords(const FRigPhysicsSolverComponent* SolverComponent)
+void FRigPhysicsSimulation::InitialiseJointRecords(
+	const URigHierarchy& Hierarchy, const FRigPhysicsSolverComponent* SolverComponent)
 {
 	ensure(JointRecords.IsEmpty());
-
-	const URigHierarchy* Hierarchy = OwningControlRig->GetHierarchy();
-	if (!Hierarchy || !SolverComponent)
+	if (!SolverComponent)
 	{
 		return;
 	}
 
 	FRigComponentKey SolverComponentKey = SolverComponent->GetKey();
 
-	TArray<FRigComponentKey> AllComponentKeys = Hierarchy->GetAllComponentKeys();
+	TArray<FRigComponentKey> AllComponentKeys = Hierarchy.GetAllComponentKeys();
 
 	for (FRigComponentKey ComponentKey : AllComponentKeys)
 	{
 		if (const FRigPhysicsJointComponent* JointComponent = Cast<FRigPhysicsJointComponent>(
-			Hierarchy->FindComponent(ComponentKey)))
+			Hierarchy.FindComponent(ComponentKey)))
 		{
 			// The authored body components may be blank, in which case we need to find the automatic ones.
 			FRigJointRecord JointRecord;
@@ -244,10 +239,10 @@ void FRigPhysicsSimulation::InitialiseJointRecords(const FRigPhysicsSolverCompon
 
 			if (!JointRecord.ChildBodyComponentKey.IsValid())
 			{
-				TArray<FRigComponentKey> SiblingComponentKeys = Hierarchy->GetComponentKeys(ComponentKey.ElementKey);
+				TArray<FRigComponentKey> SiblingComponentKeys = Hierarchy.GetComponentKeys(ComponentKey.ElementKey);
 				for (FRigComponentKey SiblingComponentKey : SiblingComponentKeys)
 				{
-					if (ShouldComponentBeInSimulation(SolverComponentKey, SiblingComponentKey))
+					if (ShouldComponentBeInSimulation(Hierarchy, SolverComponentKey, SiblingComponentKey))
 					{
 						JointRecord.ChildBodyComponentKey = SiblingComponentKey;
 						break;
@@ -257,11 +252,11 @@ void FRigPhysicsSimulation::InitialiseJointRecords(const FRigPhysicsSolverCompon
 
 			if (!JointRecord.ParentBodyComponentKey.IsValid())
 			{
-				FRigElementKey ParentElementKey = Hierarchy->GetFirstParent(ComponentKey.ElementKey);
-				TArray<FRigComponentKey> ParentComponentKeys = Hierarchy->GetComponentKeys(ParentElementKey);
+				FRigElementKey ParentElementKey = Hierarchy.GetFirstParent(ComponentKey.ElementKey);
+				TArray<FRigComponentKey> ParentComponentKeys = Hierarchy.GetComponentKeys(ParentElementKey);
 				for (FRigComponentKey ParentComponentKey : ParentComponentKeys)
 				{
-					if (ShouldComponentBeInSimulation(SolverComponentKey, ParentComponentKey))
+					if (ShouldComponentBeInSimulation(Hierarchy, SolverComponentKey, ParentComponentKey))
 					{
 						JointRecord.ParentBodyComponentKey = ParentComponentKey;
 						break;
@@ -269,9 +264,9 @@ void FRigPhysicsSimulation::InitialiseJointRecords(const FRigPhysicsSolverCompon
 				}
 			}
 
-			if (ShouldComponentBeInSimulation(SolverComponentKey, JointRecord.ChildBodyComponentKey))
+			if (ShouldComponentBeInSimulation(Hierarchy, SolverComponentKey, JointRecord.ChildBodyComponentKey))
 			{
-				if (ShouldComponentBeInSimulation(SolverComponentKey, JointRecord.ParentBodyComponentKey))
+				if (ShouldComponentBeInSimulation(Hierarchy, SolverComponentKey, JointRecord.ParentBodyComponentKey))
 				{
 					// Just make the record for now - it will be instantiated later
 					JointRecords.Add(ComponentKey, JointRecord);
@@ -283,26 +278,24 @@ void FRigPhysicsSimulation::InitialiseJointRecords(const FRigPhysicsSolverCompon
 
 
 //======================================================================================================================
-void FRigPhysicsSimulation::InitialiseBodyRecords(const FRigPhysicsSolverComponent* SolverComponent)
+void FRigPhysicsSimulation::InitialiseBodyRecords(const URigHierarchy& Hierarchy, const FRigPhysicsSolverComponent* SolverComponent)
 {
 	ensure(BodyRecords.IsEmpty());
-
-	const URigHierarchy* Hierarchy = OwningControlRig->GetHierarchy();
-	if (!Hierarchy || !SolverComponent)
+	if (!SolverComponent)
 	{
 		return;
 	}
 
 	FRigComponentKey SolverComponentKey = SolverComponent->GetKey();
 
-	TArray<FRigComponentKey> AllComponentKeys = Hierarchy->GetAllComponentKeys();
+	TArray<FRigComponentKey> AllComponentKeys = Hierarchy.GetAllComponentKeys();
 
 	// All the components in this simulation
 	TArray<FRigComponentKey> UnsortedBodyComponentKeys;
 
 	for (FRigComponentKey ComponentKey : AllComponentKeys)
 	{
-		if (ShouldComponentBeInSimulation(SolverComponentKey, ComponentKey))
+		if (ShouldComponentBeInSimulation(Hierarchy, SolverComponentKey, ComponentKey))
 		{
 			// Just make the record for now - it will be instantiated later
 			BodyRecords.Add(ComponentKey);
@@ -312,7 +305,7 @@ void FRigPhysicsSimulation::InitialiseBodyRecords(const FRigPhysicsSolverCompone
 
 	// Sort the component keys according to the traversal of their element (i.e. from root to leaf)
 	SortedBodyComponentKeys.Empty(UnsortedBodyComponentKeys.Num());
-	Hierarchy->Traverse([this, &UnsortedBodyComponentKeys](FRigBaseElement* Element, bool& bContinue)
+	Hierarchy.Traverse([this, &UnsortedBodyComponentKeys](FRigBaseElement* Element, bool& bContinue)
 		{
 			const FRigElementKey& Key = Element->GetKey();
 			for (const FRigComponentKey& ComponentKey : UnsortedBodyComponentKeys)
@@ -623,12 +616,14 @@ ImmediatePhysics::FActorHandle* FRigPhysicsSimulation::CreateBody(
 
 //======================================================================================================================
 void FRigPhysicsSimulation::Instantiate(
-	const FRigVMExecuteContext& ExecuteContext, const FRigPhysicsSolverComponent* SolverComponent)
+	const FRigVMExecuteContext&       ExecuteContext, 
+	const URigHierarchy&              Hierarchy, 
+	const FRigPhysicsSolverComponent* SolverComponent)
 {
 	DECLARE_SCOPE_HIERARCHICAL_COUNTER_FUNC()
 	QUICK_SCOPE_CYCLE_COUNTER(STAT_RigPhysics_Instantiate);
 
-	if (!bNeedToInstantiate || !OwningControlRig)
+	if (!bNeedToInstantiate)
 	{
 		return;
 	}
@@ -636,23 +631,24 @@ void FRigPhysicsSimulation::Instantiate(
 	// We need the simulation space in order to instantiate properly. This is not ideal, as we may
 	// end up updating the simulation space data twice (thus inserting data into the history). This
 	// shouldn't really matter as it will only be on the first step.
-	UpdateSimulationSpaceStateAndCalculateData(ExecuteContext, SolverComponent, 0.0f);
+	UpdateSimulationSpaceStateAndCalculateData(
+		ExecuteContext, Hierarchy, SolverComponent, 0.0f, ExecuteContext.GetAbsoluteTime());
 
 	InitialiseSimulation(SolverComponent);
 
-	InitialiseBodyRecords(SolverComponent);
+	InitialiseBodyRecords(Hierarchy, SolverComponent);
 
-	InitialiseJointRecords(SolverComponent);
+	InitialiseJointRecords(Hierarchy, SolverComponent);
 
-	InitialiseControlRecords(SolverComponent);
+	InitialiseControlRecords(Hierarchy, SolverComponent);
 
 	FRigPhysicsIgnorePairs IgnorePairs;
 
-	InstantiatePhysicsBodies(SolverComponent, IgnorePairs);
+	InstantiatePhysicsBodies(Hierarchy, SolverComponent, IgnorePairs);
 
-	InstantiatePhysicsJoints(SolverComponent, IgnorePairs);
+	InstantiatePhysicsJoints(Hierarchy, SolverComponent, IgnorePairs);
 
-	InstantiateControls(SolverComponent, IgnorePairs);
+	InstantiateControls(Hierarchy, SolverComponent, IgnorePairs);
 
 	// This is done last as it applies IgnorePairs
 	InstantiateSolverCollision(SolverComponent, IgnorePairs);
@@ -710,14 +706,13 @@ void FRigPhysicsSimulation::InstantiateSolverCollision(
 
 //======================================================================================================================
 void FRigPhysicsSimulation::InstantiatePhysicsBodies(
+	const URigHierarchy&              Hierarchy,
 	const FRigPhysicsSolverComponent* SolverComponent,
 	FRigPhysicsIgnorePairs&           IgnorePairs)
 {
 	DECLARE_SCOPE_HIERARCHICAL_COUNTER_FUNC()
 	QUICK_SCOPE_CYCLE_COUNTER(STAT_RigPhysics_InstantiatePhysicsBodies);
 
-	const URigHierarchy* Hierarchy = OwningControlRig->GetHierarchy();
-
 	const FRigPhysicsSolverSettings& SolverSettings = SolverComponent->SolverSettings;
 	const FRigPhysicsSimulationSpaceSettings& SimulationSpaceSettings = SolverComponent->SimulationSpaceSettings;
 
@@ -727,7 +722,7 @@ void FRigPhysicsSimulation::InstantiatePhysicsBodies(
 		FRigBodyRecord& Record = BodyRecordPair.Value;
 
 		if (const FRigPhysicsBodyComponent* PhysicsComponent = Cast<FRigPhysicsBodyComponent>(
-			Hierarchy->FindComponent(ComponentKey)))
+			Hierarchy.FindComponent(ComponentKey)))
 		{
 			const FRigPhysicsCollision& Collision = PhysicsComponent->Collision;
 			const FRigPhysicsDynamics& Dynamics = PhysicsComponent->Dynamics;
@@ -742,7 +737,7 @@ void FRigPhysicsSimulation::InstantiatePhysicsBodies(
 			// What should we do if the key is not valid? 
 			if (SourceKey.IsValid())
 			{
-				FTransform SourceComponentSpaceTM = Hierarchy->GetGlobalTransform(SourceKey);
+				FTransform SourceComponentSpaceTM = Hierarchy.GetGlobalTransform(SourceKey);
 				const FTransform SourceSimulationSpaceTM =
 					ConvertComponentSpaceTransformToSimSpace(SolverComponent->SolverSettings, SourceComponentSpaceTM);
 				Record.ActorHandle = CreateBody(
@@ -765,14 +760,13 @@ void FRigPhysicsSimulation::InstantiatePhysicsBodies(
 
 //======================================================================================================================
 void FRigPhysicsSimulation::InstantiatePhysicsJoints(
+	const URigHierarchy&              Hierarchy,
 	const FRigPhysicsSolverComponent* SolverComponent,
 	FRigPhysicsIgnorePairs&           IgnorePairs)
 {
 	DECLARE_SCOPE_HIERARCHICAL_COUNTER_FUNC()
 	QUICK_SCOPE_CYCLE_COUNTER(STAT_RigPhysics_InstantiatePhysicsJoints);
 
-	const URigHierarchy* Hierarchy = OwningControlRig->GetHierarchy();
-
 	const FRigPhysicsSolverSettings& SolverSettings = SolverComponent->SolverSettings;
 	const FRigPhysicsSimulationSpaceSettings& SimulationSpaceSettings = SolverComponent->SimulationSpaceSettings;
 
@@ -783,7 +777,7 @@ void FRigPhysicsSimulation::InstantiatePhysicsJoints(
 		FRigJointRecord& JointRecord = JointRecordPair.Value;
 
 		if (const FRigPhysicsJointComponent* PhysicsJointComponent = Cast<FRigPhysicsJointComponent>(
-			Hierarchy->FindComponent(JointComponentKey)))
+			Hierarchy.FindComponent(JointComponentKey)))
 		{
 			const FRigPhysicsJointData& JointData = PhysicsJointComponent->JointData;
 
@@ -795,7 +789,7 @@ void FRigPhysicsSimulation::InstantiatePhysicsJoints(
 
 			// Joints require both parent and child to exist
 			const FRigPhysicsBodyComponent* ChildPhysicsComponent =
-				Cast<FRigPhysicsBodyComponent>(Hierarchy->FindComponent(ChildBodyComponentKey));
+				Cast<FRigPhysicsBodyComponent>(Hierarchy.FindComponent(ChildBodyComponentKey));
 			FRigBodyRecord* ChildBodyRecord = BodyRecords.Find(ChildBodyComponentKey);
 			ImmediatePhysics::FActorHandle* ChildActorHandle = ChildBodyRecord ? ChildBodyRecord->ActorHandle : nullptr;
 			if (!ChildActorHandle)
@@ -804,7 +798,7 @@ void FRigPhysicsSimulation::InstantiatePhysicsJoints(
 			}
 
 			const FRigPhysicsBodyComponent* ParentPhysicsComponent =
-				Cast<FRigPhysicsBodyComponent>(Hierarchy->FindComponent(ParentBodyComponentKey));
+				Cast<FRigPhysicsBodyComponent>(Hierarchy.FindComponent(ParentBodyComponentKey));
 			FRigBodyRecord* ParentBodyRecord = BodyRecords.Find(ParentBodyComponentKey);
 			ImmediatePhysics::FActorHandle* ParentActorHandle = ParentBodyRecord ? ParentBodyRecord->ActorHandle : nullptr;
 			if (!ParentActorHandle)
@@ -828,7 +822,7 @@ void FRigPhysicsSimulation::InstantiatePhysicsJoints(
 
 				if (JointData.bAutoCalculateParentOffset)
 				{
-					JointSettings.ConnectorTransforms[1] = Hierarchy->GetLocalTransform(ChildBoneKey, true);
+					JointSettings.ConnectorTransforms[1] = Hierarchy.GetLocalTransform(ChildBoneKey, true);
 				}
 				JointSettings.ConnectorTransforms[1] =
 					JointData.ExtraParentOffset * JointSettings.ConnectorTransforms[1];
@@ -867,21 +861,20 @@ void FRigPhysicsSimulation::InstantiatePhysicsJoints(
 
 //======================================================================================================================
 void FRigPhysicsSimulation::InstantiateControls(
+	const URigHierarchy&              Hierarchy,
 	const FRigPhysicsSolverComponent* SolverComponent,
 	FRigPhysicsIgnorePairs&           IgnorePairs)
 {
 	DECLARE_SCOPE_HIERARCHICAL_COUNTER_FUNC()
 	QUICK_SCOPE_CYCLE_COUNTER(STAT_RigPhysics_InstantiateControls);
 
-	const URigHierarchy* Hierarchy = OwningControlRig->GetHierarchy();
-
 	for (TPair<FRigComponentKey, FRigControlRecord>& ControlRecordPair : ControlRecords)
 	{
 		const FRigComponentKey& ComponentKey = ControlRecordPair.Key;
 		FRigControlRecord& ControlRecord = ControlRecordPair.Value;
 
 		if (const FRigPhysicsControlComponent* ControlComponent = Cast<FRigPhysicsControlComponent>(
-			Hierarchy->FindComponent(ComponentKey)))
+			Hierarchy.FindComponent(ComponentKey)))
 		{
 			if (FRigBodyRecord* ChildBodyRecord = BodyRecords.Find(ControlRecord.ChildBodyComponentKey))
 			{
@@ -913,13 +906,13 @@ void FRigPhysicsSimulation::InstantiateControls(
 
 //======================================================================================================================
 void FRigPhysicsSimulation::UpdateBodyRecordPrePhysics(
+	const URigHierarchy&              Hierarchy,
 	const FRigPhysicsSolverComponent* SolverComponent,
 	const float                       DeltaTime,
 	FRigBodyRecord&                   Record, 
 	const FRigPhysicsBodyComponent*   PhysicsComponent)
 {
-	URigHierarchy* Hierarchy = OwningControlRig->GetHierarchy();
-	if (!Record.ActorHandle || !Hierarchy)
+	if (!Record.ActorHandle)
 	{
 		return;
 	}
@@ -937,7 +930,7 @@ void FRigPhysicsSimulation::UpdateBodyRecordPrePhysics(
 	}
 	if (SourceKey.IsValid())
 	{
-		Record.SourceComponentSpaceTM = Hierarchy->GetGlobalTransform(SourceKey);
+		Record.SourceComponentSpaceTM = Hierarchy.GetGlobalTransform(SourceKey);
 	}
 
 	if (UpdateCounter == PreviousUpdateCounter + 1)
@@ -1060,9 +1053,9 @@ void FRigPhysicsSimulation::UpdateBodyPrePhysics(
 
 //======================================================================================================================
 void FRigPhysicsSimulation::UpdateJointPrePhysics(
+	const URigHierarchy&                Hierarchy, 
 	FRigJointRecord&                    Record, 
 	const FRigPhysicsJointComponent*    PhysicsJointComponent, 
-	const URigHierarchy*                Hierarchy, 
 	const float                         DeltaTime)
 {
 	// Now update the joint targets
@@ -1097,7 +1090,7 @@ void FRigPhysicsSimulation::UpdateJointPrePhysics(
 
 			FRigElementKey ChildSourceKey;
 			if (const FRigPhysicsBodyComponent* ChildPhysicsComponent = Cast<FRigPhysicsBodyComponent>(
-				Hierarchy->FindComponent(Record.ChildBodyComponentKey)))
+				Hierarchy.FindComponent(Record.ChildBodyComponentKey)))
 			{
 				ChildSourceKey = ChildPhysicsComponent->BodySolverSettings.SourceBone;
 				if (!ChildSourceKey.IsValid())
@@ -1112,7 +1105,7 @@ void FRigPhysicsSimulation::UpdateJointPrePhysics(
 			{
 				FRigElementKey ParentSourceKey;
 				if (const FRigPhysicsBodyComponent* ParentPhysicsComponent = Cast<FRigPhysicsBodyComponent>(
-					Hierarchy->FindComponent(Record.ParentBodyComponentKey)))
+					Hierarchy.FindComponent(Record.ParentBodyComponentKey)))
 				{
 					ParentSourceKey = ParentPhysicsComponent->BodySolverSettings.SourceBone;
 					if (!ParentSourceKey.IsValid())
@@ -1131,8 +1124,8 @@ void FRigPhysicsSimulation::UpdateJointPrePhysics(
 
 				// Note that the drive operates between a parent and child part, so we
 				// don't need to worry about global/component (etc) space.
-				const FTransform ChildTM = Hierarchy->GetGlobalTransform(ChildSourceKey);
-				const FTransform ParentTM = Hierarchy->GetGlobalTransform(ParentSourceKey);
+				const FTransform ChildTM = Hierarchy.GetGlobalTransform(ChildSourceKey);
+				const FTransform ParentTM = Hierarchy.GetGlobalTransform(ParentSourceKey);
 
 				const UE::PhysicsControl::FPosQuat ComponentSpaceParentFrameTM = 
 					Settings.ConnectorTransforms[1] * ParentTM;
@@ -1181,12 +1174,12 @@ void FRigPhysicsSimulation::UpdateJointPrePhysics(
 
 //======================================================================================================================
 static UE::PhysicsControl::FPosQuat CalculateTargetTM(
-	const URigHierarchy*               Hierarchy,
+	const URigHierarchy&               Hierarchy,
 	const Chaos::FPBDJointSettings&    JointSettings,
 	const FRigControlRecord&           Record)
 {
 	const FTransform ChildTM = Record.ChildBodyComponentKey.IsValid() ?
-		Hierarchy->GetGlobalTransform(Record.ChildBodyComponentKey.ElementKey) :
+		Hierarchy.GetGlobalTransform(Record.ChildBodyComponentKey.ElementKey) :
 		FTransform();
 
 	const UE::PhysicsControl::FPosQuat ChildTargetTM =
@@ -1195,7 +1188,7 @@ static UE::PhysicsControl::FPosQuat CalculateTargetTM(
 
 	if (Record.ParentBodyComponentKey.IsValid())
 	{
-		const FTransform ParentTM = Hierarchy->GetGlobalTransform(Record.ParentBodyComponentKey.ElementKey);
+		const FTransform ParentTM = Hierarchy.GetGlobalTransform(Record.ParentBodyComponentKey.ElementKey);
 		const UE::PhysicsControl::FPosQuat ParentTargetTM =
 			UE::PhysicsControl::FPosQuat(ParentTM) *
 			UE::PhysicsControl::FPosQuat(JointSettings.ConnectorTransforms[ConstraintParentIndex]);
@@ -1206,13 +1199,11 @@ static UE::PhysicsControl::FPosQuat CalculateTargetTM(
 
 //======================================================================================================================
 void FRigPhysicsSimulation::CheckForResetsPrePhysics(
-	FRigPhysicsSolverComponent* SolverComponent, const float DeltaTime)
+	const URigHierarchy& Hierarchy, FRigPhysicsSolverComponent* SolverComponent, const float DeltaTime)
 {
 	DECLARE_SCOPE_HIERARCHICAL_COUNTER_FUNC()
 	QUICK_SCOPE_CYCLE_COUNTER(STAT_RigPhysics_CheckForResetsPrePhysics);
 
-	URigHierarchy* Hierarchy = OwningControlRig->GetHierarchy();
-
 	const FRigPhysicsSolverSettings& SolverSettings = SolverComponent->SolverSettings;
 
 	double SpeedThresholdForResetSquared = SolverSettings.KinematicSpeedThresholdForReset > 0 ?
@@ -1238,7 +1229,7 @@ void FRigPhysicsSimulation::CheckForResetsPrePhysics(
 		if (Record.ActorHandle && Record.ActorHandle->GetIsKinematic())
 		{
 			if (const FRigPhysicsBodyComponent* PhysicsComponent = Cast<FRigPhysicsBodyComponent>(
-				Hierarchy->FindComponent(ComponentKey)))
+				Hierarchy.FindComponent(ComponentKey)))
 			{
 				if (PhysicsComponent->BodySolverSettings.bIncludeInChecksForReset)
 				{
@@ -1259,12 +1250,12 @@ void FRigPhysicsSimulation::CheckForResetsPrePhysics(
 		if (HighestSpeedSq > SpeedThresholdForResetSquared)
 		{
 			UE_LOG(LogRigPhysics, Log, TEXT("Speed %f triggered reset in %s"),
-				FMath::Sqrt(HighestSpeedSq), *OwningControlRig->GetName());
+				FMath::Sqrt(HighestSpeedSq), *OwnerName.ToString());
 		}
 		if (HighestAccelerationSq > AccelerationThresholdForResetSquared)
 		{
 			UE_LOG(LogRigPhysics, Log, TEXT("Acceleration %f triggered reset in %s"),
-				FMath::Sqrt(HighestAccelerationSq), *OwningControlRig->GetName());
+				FMath::Sqrt(HighestAccelerationSq), *OwnerName.ToString());
 		}
 		SolverComponent->TrackInputCounter = FMath::Max(SolverComponent->TrackInputCounter, 3);
 	}
@@ -1273,14 +1264,13 @@ void FRigPhysicsSimulation::CheckForResetsPrePhysics(
 //======================================================================================================================
 void FRigPhysicsSimulation::UpdatePrePhysics(
 	const FRigVMExecuteContext& ExecuteContext, 
+	const URigHierarchy&        Hierarchy,
 	FRigPhysicsSolverComponent* SolverComponent, 
 	const float                 DeltaTime)
 {
 	DECLARE_SCOPE_HIERARCHICAL_COUNTER_FUNC()
 	QUICK_SCOPE_CYCLE_COUNTER(STAT_RigPhysics_UpdatePrePhysics);
 
-	URigHierarchy* Hierarchy = OwningControlRig->GetHierarchy();
-
 	if (CollisionActorHandle)
 	{
 		FTransform BodyRelSimSpaceTM = ConvertCollisionSpaceTransformToSimSpace(
@@ -1293,17 +1283,17 @@ void FRigPhysicsSimulation::UpdatePrePhysics(
 		const FRigComponentKey& ComponentKey = BodyRecordPair.Key;
 		FRigBodyRecord& Record = BodyRecordPair.Value;
 		if (const FRigPhysicsBodyComponent* PhysicsComponent = Cast<FRigPhysicsBodyComponent>(
-			Hierarchy->FindComponent(ComponentKey)))
+			Hierarchy.FindComponent(ComponentKey)))
 		{
-			UpdateBodyRecordPrePhysics(SolverComponent, DeltaTime, Record, PhysicsComponent);
+			UpdateBodyRecordPrePhysics(Hierarchy, SolverComponent, DeltaTime, Record, PhysicsComponent);
 		}
 	}
 
-	CheckForResetsPrePhysics(SolverComponent, DeltaTime);
+	CheckForResetsPrePhysics(Hierarchy, SolverComponent, DeltaTime);
 	if (SolverComponent->TrackInputCounter > 0)
 	{
 		UE_LOG(LogRigPhysics, Log, TEXT("Forcing tracking (counter = %d) of input for %s"), 
-			SolverComponent->TrackInputCounter, *OwningControlRig->GetName());
+			SolverComponent->TrackInputCounter, *OwnerName.ToString());
 	}
 
 	{
@@ -1314,7 +1304,7 @@ void FRigPhysicsSimulation::UpdatePrePhysics(
 			const FRigComponentKey& ComponentKey = BodyRecordPair.Key;
 			FRigBodyRecord& Record = BodyRecordPair.Value;
 			if (const FRigPhysicsBodyComponent* PhysicsComponent = Cast<FRigPhysicsBodyComponent>(
-				Hierarchy->FindComponent(ComponentKey)))
+				Hierarchy.FindComponent(ComponentKey)))
 			{
 				UpdateBodyPrePhysics(ExecuteContext, SolverComponent, Record, PhysicsComponent);
 			}
@@ -1330,9 +1320,9 @@ void FRigPhysicsSimulation::UpdatePrePhysics(
 			FRigJointRecord& JointRecord = JointRecordPair.Value;
 
 			if (const FRigPhysicsJointComponent* JointComponent = Cast<FRigPhysicsJointComponent>(
-				Hierarchy->FindComponent(ComponentKey)))
+				Hierarchy.FindComponent(ComponentKey)))
 			{
-				UpdateJointPrePhysics(JointRecord, JointComponent, Hierarchy, DeltaTime);
+				UpdateJointPrePhysics(Hierarchy, JointRecord, JointComponent, DeltaTime);
 			}
 		}
 	}
@@ -1346,7 +1336,7 @@ void FRigPhysicsSimulation::UpdatePrePhysics(
 			FRigControlRecord& ControlRecord = ControlRecordPair.Value;
 
 			if (const FRigPhysicsControlComponent* ControlComponent = Cast<FRigPhysicsControlComponent>(
-				Hierarchy->FindComponent(ComponentKey)))
+				Hierarchy.FindComponent(ComponentKey)))
 			{
 				if (ControlComponent->ControlData.bEnabled)
 				{
@@ -1363,7 +1353,7 @@ void FRigPhysicsSimulation::UpdateControlPrePhysics(
 	FRigControlRecord&                 ControlRecord,
 	const FRigPhysicsControlComponent* ControlComponent,
 	const FRigPhysicsSolverComponent*  SolverComponent,
-	const URigHierarchy*               Hierarchy,
+	const URigHierarchy&               Hierarchy,
 	const float                        DeltaTime)
 {
 	const FRigPhysicsSolverSettings& SolverSettings = SolverComponent->SolverSettings;
@@ -1476,6 +1466,7 @@ void FRigPhysicsSimulation::UpdateControlPrePhysics(
 // Note that we read back into a target bone, which may have been specified explicitly, or will
 // otherwise default to the physics element parent.
 void FRigPhysicsSimulation::UpdatePostPhysics(
+	URigHierarchy&              Hierarchy,
 	FRigPhysicsSolverComponent* SolverComponent, 
 	const float                 Alpha, 
 	const float                 DeltaTime,
@@ -1490,7 +1481,6 @@ void FRigPhysicsSimulation::UpdatePostPhysics(
 	}
 	const FRigPhysicsSolverSettings& SolverSettings = SolverComponent->SolverSettings;
 	bool bGotInvalidSimulationData = false;
-	URigHierarchy* Hierarchy = OwningControlRig->GetHierarchy();
 
 	double PositionThresholdForResetSquared = SolverSettings.PositionThresholdForReset > 0 ?
 		FMath::Square(SolverSettings.PositionThresholdForReset) : TNumericLimits<double>::Max();
@@ -1537,7 +1527,7 @@ void FRigPhysicsSimulation::UpdatePostPhysics(
 		if (HighestPosition > 0)
 		{
 			UE_LOG(LogRigPhysics, Log, TEXT("Position %f triggered teleport in %s - resetting pose"),
-				HighestPosition, *OwningControlRig->GetName());
+				HighestPosition, *OwnerName.ToString());
 		}
 		// Avoid cached transforms being used in controls by bumping the update counter. 
 		UpdateCounter += 1;
@@ -1551,14 +1541,14 @@ void FRigPhysicsSimulation::UpdatePostPhysics(
 	if (HighestPosition > 0)
 	{
 		UE_LOG(LogRigPhysics, Log, TEXT("Resetting state to input pose in %s"), 
-			*OwningControlRig->GetName());
+			*OwnerName.ToString());
 		for (TPair<FRigComponentKey, FRigBodyRecord>& BodyRecordPair : BodyRecords)
 		{
 			const FRigComponentKey& ComponentKey = BodyRecordPair.Key;
 			const FRigBodyRecord& Record = BodyRecordPair.Value;
 
 			if (const FRigPhysicsBodyComponent* PhysicsComponent = Cast<FRigPhysicsBodyComponent>(
-				Hierarchy->FindComponent(ComponentKey)))
+				Hierarchy.FindComponent(ComponentKey)))
 			{
 				if (Record.ActorHandle && !Record.ActorHandle->GetIsKinematic())
 				{
@@ -1609,7 +1599,7 @@ void FRigPhysicsSimulation::UpdatePostPhysics(
 					// children), so that attached animation bones will follow physics, but we
 					// rely on our bodies being sorted so we work out from the leaf nodes so as
 					// not to disturb previously set bodies.
-					Hierarchy->SetGlobalTransform(Record.TargetElementKey, Record.FinalComponentSpaceTM, false, true);
+					Hierarchy.SetGlobalTransform(Record.TargetElementKey, Record.FinalComponentSpaceTM, false, true);
 				}
 			}
 		}
@@ -1619,6 +1609,7 @@ void FRigPhysicsSimulation::UpdatePostPhysics(
 //======================================================================================================================
 void FRigPhysicsSimulation::StepSimulation(
 	const FRigVMExecuteContext& ExecuteContext,
+	URigHierarchy&              Hierarchy,
 	FRigPhysicsSolverComponent* SolverComponent,
 	const float                 DeltaTimeOverride,
 	const float                 SimulationSpaceDeltaTimeOverride,
@@ -1648,12 +1639,12 @@ void FRigPhysicsSimulation::StepSimulation(
 
 	// We need to know about the simulation space etc before we can instantiate anything into the right place
 	SimulationSpaceData = UpdateSimulationSpaceStateAndCalculateData(
-		ExecuteContext, SolverComponent, PhysicsSimulationSpaceDeltaTime);
+		ExecuteContext, Hierarchy, SolverComponent, PhysicsSimulationSpaceDeltaTime, ExecuteContext.GetAbsoluteTime());
 
 	// We instantiate when we do the first simulation - this makes sure any changes applied by
 	// the user have been made. It also means there is no overhead if physics is never stepped.
 	// However, there may be a hitch due to the creation, so it may also happen during construction.
-	Instantiate(ExecuteContext, SolverComponent); // There is an early out if it's already been done
+	Instantiate(ExecuteContext, Hierarchy, SolverComponent); // There is an early out if it's already been done
 
 	if (!Simulation)
 	{
@@ -1718,7 +1709,7 @@ void FRigPhysicsSimulation::StepSimulation(
 	// * We can't actually simulate with dt = 0
 	if (PhysicsDeltaTime > 0.0f)
 	{
-		UpdatePrePhysics(ExecuteContext, SolverComponent, PhysicsDeltaTime);
+		UpdatePrePhysics(ExecuteContext, Hierarchy, SolverComponent, PhysicsDeltaTime);
 
 		Simulation->Simulate(
 			PhysicsDeltaTime, MaxDeltaTime, MaxTimeSteps,
@@ -1728,7 +1719,7 @@ void FRigPhysicsSimulation::StepSimulation(
 	}
 
 	// Always do a read-back, even for zero Dt
-	UpdatePostPhysics(SolverComponent, FMath::Clamp(Alpha, 0.0f, 1.0f), PhysicsDeltaTime, FixedTimeStep > 0);
+	UpdatePostPhysics(Hierarchy, SolverComponent, FMath::Clamp(Alpha, 0.0f, 1.0f), PhysicsDeltaTime, FixedTimeStep > 0);
 
 	if (SolverComponent->TrackInputCounter > 0)
 	{
diff --git a/Engine/Plugins/Experimental/ControlRigPhysics/Source/ControlRigPhysics/Private/RigPhysicsSimulation.h b/Engine/Plugins/Experimental/ControlRigPhysics/Source/ControlRigPhysics/Private/RigPhysicsSimulation.h
--- a/Engine/Plugins/Experimental/ControlRigPhysics/Source/ControlRigPhysics/Private/RigPhysicsSimulation.h
+++ b/Engine/Plugins/Experimental/ControlRigPhysics/Source/ControlRigPhysics/Private/RigPhysicsSimulation.h
@@ -142,11 +142,13 @@ struct FRigPhysicsSimulation : public FRigPhysicsSimulationBase
 public:
 	GENERATED_BODY()
 
-	FRigPhysicsSimulation(
-		TObjectPtr<UControlRig> InOwningControlRig = nullptr);
+	FRigPhysicsSimulation(const FName InOwnerName = FName());
 
 	// This will initialise/create the simulation and then create everything we need in it. 
-	void Instantiate(const FRigVMExecuteContext& ExecuteContext, const FRigPhysicsSolverComponent* SolverComponent);
+	void Instantiate(
+		const FRigVMExecuteContext&       ExecuteContext, 
+		const URigHierarchy&              Hierarchy, 
+		const FRigPhysicsSolverComponent* SolverComponent);
 
 	// Integrates the simulation forwards.
 	// If DeltaTimeOverride is +ve, then that value is used.
@@ -154,7 +156,8 @@ public:
 	// If it is negative, then the simulation isn't stepped.
 	void StepSimulation(
 		const FRigVMExecuteContext& ExecuteContext, 
-		FRigPhysicsSolverComponent* SolverComponent, 
+		URigHierarchy&              Hierarchy,
+		FRigPhysicsSolverComponent* SolverComponent,
 		const float                 DeltaTimeOverride, 
 		const float                 SimulationSpaceDeltaTimeOverride,
 		const float                 Alpha);
@@ -204,26 +207,29 @@ private:
 	void InitialiseSimulation(const FRigPhysicsSolverComponent* SolverComponent);
 
 	// This collects all the bodies associated with the solver and makes records for them
-	void InitialiseBodyRecords(const FRigPhysicsSolverComponent* SolverComponent);
+	void InitialiseBodyRecords(const URigHierarchy& Hierarchy, const FRigPhysicsSolverComponent* SolverComponent);
 
 	// This collects all the joints associated with the solver and makes records for them
-	void InitialiseJointRecords(const FRigPhysicsSolverComponent* SolverComponent);
+	void InitialiseJointRecords(const URigHierarchy& Hierarchy, const FRigPhysicsSolverComponent* SolverComponent);
 
 	// This collects all the bodies associated with the solver and makes records for them
-	void InitialiseControlRecords(const FRigPhysicsSolverComponent* SolverComponent);
+	void InitialiseControlRecords(const URigHierarchy& Hierarchy, const FRigPhysicsSolverComponent* SolverComponent);
 
 	// Creates low level bodies
 	void InstantiatePhysicsBodies(
+		const URigHierarchy&              Hierarchy,
 		const FRigPhysicsSolverComponent* SolverComponent,
 		FRigPhysicsIgnorePairs&           IgnorePairs);
 
 	// Creates low level physics joints 
 	void InstantiatePhysicsJoints(
+		const URigHierarchy&              Hierarchy,
 		const FRigPhysicsSolverComponent* SolverComponent,
 		FRigPhysicsIgnorePairs&           IgnorePairs);
 
 	// Creates low level controls
 	void InstantiateControls(
+		const URigHierarchy&              Hierarchy,
 		const FRigPhysicsSolverComponent* SolverComponent,
 		FRigPhysicsIgnorePairs&           IgnorePairs);
 
@@ -235,7 +241,7 @@ private:
 	// Returns true if the component is physics, and its solver matches the solver component
 	// (directly, or automatically)
 	bool ShouldComponentBeInSimulation(
-		FRigComponentKey SolverComponentKey, FRigComponentKey ComponentKey) const;
+		const URigHierarchy& Hierarchy, FRigComponentKey SolverComponentKey, FRigComponentKey ComponentKey) const;
 
 	// Creates an actor with collision. This will by dynamic if Dynamics is valid, or otherwise kinematic
 	ImmediatePhysics::FActorHandle* CreateBody(
@@ -250,36 +256,39 @@ private:
 	void DestroyPhysicsSimulation();
 
 	void UpdatePrePhysics(
-		const FRigVMExecuteContext&       ExecuteContext, 
-		FRigPhysicsSolverComponent*       SolverComponent, 
-		const float                       DeltaTime);
+		const FRigVMExecuteContext&        ExecuteContext, 
+		const URigHierarchy&               Hierarchy,
+		FRigPhysicsSolverComponent*        SolverComponent, 
+		const float                        DeltaTime);
 
 	void UpdateBodyRecordPrePhysics(
-		const FRigPhysicsSolverComponent* SolverComponent,
-		const float                       DeltaTime,
-		FRigBodyRecord&                   Record,
-		const FRigPhysicsBodyComponent*   PhysicsComponent);
+		const URigHierarchy&               Hierarchy,
+		const FRigPhysicsSolverComponent*  SolverComponent,
+		const float                        DeltaTime,
+		FRigBodyRecord&                    Record,
+		const FRigPhysicsBodyComponent*    PhysicsComponent);
 
 	void UpdateBodyPrePhysics(
-		const FRigVMExecuteContext&       ExecuteContext, 
-		const FRigPhysicsSolverComponent* SolverComponent,
-		const FRigBodyRecord&             Record, 
-		const FRigPhysicsBodyComponent*   PhysicsComponent);
+		const FRigVMExecuteContext&        ExecuteContext, 
+		const FRigPhysicsSolverComponent*  SolverComponent,
+		const FRigBodyRecord&              Record, 
+		const FRigPhysicsBodyComponent*    PhysicsComponent);
 
 	void UpdateJointPrePhysics(
-		FRigJointRecord&                 JointRecord,
-		const FRigPhysicsJointComponent* JointComponent, 
-		const URigHierarchy*             Hierarchy, 
-		const float                      DeltaTime);
+		const URigHierarchy&               Hierarchy, 
+		FRigJointRecord&                   JointRecord,
+		const FRigPhysicsJointComponent*   JointComponent, 
+		const float                        DeltaTime);
 
 	void UpdateControlPrePhysics(
 		FRigControlRecord&                 ControlRecord,
 		const FRigPhysicsControlComponent* ControlComponent,
 		const FRigPhysicsSolverComponent*  SolverComponent,
-		const URigHierarchy*               Hierarchy,
+		const URigHierarchy&               Hierarchy,
 		const float                        DeltaTime);
 
 	void UpdatePostPhysics(
+		URigHierarchy&              Hierarchy,
 		FRigPhysicsSolverComponent* SolverComponent, 
 		const float                 Alpha, 
 		const float                 DeltaTime,
@@ -287,12 +296,16 @@ private:
 
 	void InitSimulationSpace(const FTransform& ComponentToWorld, const FTransform& BoneToComponent);
 
-	void CheckForResetsPrePhysics(FRigPhysicsSolverComponent* SolverComponent, const float DeltaTime);
+	void CheckForResetsPrePhysics(
+		const URigHierarchy& Hierarchy, FRigPhysicsSolverComponent* SolverComponent, const float DeltaTime);
 
+	// AbsoluteTime is used for calculating turbulence when applying the "wind"
 	FSimulationSpaceData UpdateSimulationSpaceStateAndCalculateData(
-		const FRigVMExecuteContext&       ExecuteContext,
-		const FRigPhysicsSolverComponent* SolverComponent, 
-		const float                       Dt);
+		const FRigVMExecuteContext&         ExecuteContext,
+		const URigHierarchy&                Hierarchy,
+		const FRigPhysicsSolverComponent*   SolverComponent,
+		const float                         Dt,
+		const double                        AbsoluteTime);
 
 	// Returns the simulation space transform, in world space
 	FTransform GetSimulationSpaceTransform(const FRigPhysicsSolverSettings& SolverSettings) const;
@@ -328,6 +341,9 @@ private:
 	TObjectPtr<UControlRig> OwningControlRig;
 	FRigComponentKey        PhysicsSolverComponentKey;
 
+	// A name that identifies what owns us - just used for debugging/logging
+	FName                   OwnerName;
+
 	// All the bodies, but in no particular order
 	TMap<FRigComponentKey, FRigBodyRecord> BodyRecords;
 
diff --git a/Engine/Plugins/Experimental/ControlRigPhysics/Source/ControlRigPhysics/Private/RigPhysicsSimulationSpace.cpp b/Engine/Plugins/Experimental/ControlRigPhysics/Source/ControlRigPhysics/Private/RigPhysicsSimulationSpace.cpp
--- a/Engine/Plugins/Experimental/ControlRigPhysics/Source/ControlRigPhysics/Private/RigPhysicsSimulationSpace.cpp
+++ b/Engine/Plugins/Experimental/ControlRigPhysics/Source/ControlRigPhysics/Private/RigPhysicsSimulationSpace.cpp
@@ -5,8 +5,6 @@
 
 #include "Rigs/RigHierarchy.h"
 #include "Rigs/RigHierarchyController.h"
-#include "RigVMCore/RigVMExecuteContext.h"
-#include "ControlRig.h"
 
 #include "PhysicsControlHelpers.h"
 
@@ -212,26 +210,20 @@ namespace FPerlinHelpers
 //======================================================================================================================
 // Note - don't use the space conversion functions here as the state won't have been set yet.
 FRigPhysicsSimulation::FSimulationSpaceData FRigPhysicsSimulation::UpdateSimulationSpaceStateAndCalculateData(
-	const FRigVMExecuteContext&       ExecuteContext,
-	const FRigPhysicsSolverComponent* SolverComponent, 
-	const float                       Dt)
+	const FRigVMExecuteContext&         ExecuteContext, 
+	const URigHierarchy&                Hierarchy,
+	const FRigPhysicsSolverComponent* SolverComponent,
+	const float                         Dt,
+	const double                        AbsoluteTime)
 {
 	const FRigPhysicsSolverSettings& SolverSettings = SolverComponent->SolverSettings;
 	const FRigPhysicsSimulationSpaceSettings& SimulationSpaceSettings = SolverComponent->SimulationSpaceSettings;
 
-	if (const USceneComponent* OwningSceneComponent = OwningControlRig->GetOwningSceneComponent())
-	{
-		SimulationSpaceState.ComponentTM = OwningSceneComponent->GetComponentTransform();
-	}
-	else 
-	{
-		SimulationSpaceState.ComponentTM.SetIdentity();
-	}
+	SimulationSpaceState.ComponentTM = ExecuteContext.GetToWorldSpaceTransform();
 
 	if (SolverSettings.SimulationSpace == ERigPhysicsSimulationSpace::SpaceBone &&  SolverSettings.SpaceBone.IsValid())
 	{
-		URigHierarchy* Hierarchy = OwningControlRig->GetHierarchy();
-		SimulationSpaceState.BoneRelComponentTM = Hierarchy->GetGlobalTransform(SolverSettings.SpaceBone);
+		SimulationSpaceState.BoneRelComponentTM = Hierarchy.GetGlobalTransform(SolverSettings.SpaceBone);
 	}
 
 	// Record the history - but avoid polluting it with zero Dt updates. What that means is - if we
@@ -332,28 +324,28 @@ FRigPhysicsSimulation::FSimulationSpaceData FRigPhysicsSimulation::UpdateSimulat
 			UE_LOG(LogRigPhysics, Log, TEXT("Detected linear Acceleration (%f > %f) teleport in %s"),
 				SimulationSpaceData.LinearAcceleration.Length(), 
 				SimulationSpaceSettings.LinearAccelerationThresholdForTeleport,
-				*OwningControlRig->GetName());
+				*OwnerName.ToString());
 		}
 		if (bAngularAccelerationTrigger)
 		{
 			UE_LOG(LogRigPhysics, Log, TEXT("Detected angular Acceleration (%f > %f) teleport in %s"),
 				FMath::RadiansToDegrees(SimulationSpaceData.AngularAcceleration.Length()),
 				SimulationSpaceSettings.AngularAccelerationThresholdForTeleport,
-				*OwningControlRig->GetName());
+				*OwnerName.ToString());
 		}
 		if (bPositionTrigger)
 		{
 			UE_LOG(LogRigPhysics, Log, TEXT("Detected position (%f > %f) teleport in %s"),
 				SimulationSpaceData.LinearVelocity.Length() * SimulationSpaceState.Dt,
 				SimulationSpaceSettings.PositionChangeThresholdForTeleport,
-				*OwningControlRig->GetName());
+				*OwnerName.ToString());
 		}
 		if (bOrientationTrigger)
 		{
 			UE_LOG(LogRigPhysics, Log, TEXT("Detected orientation (%f > %f) teleport in %s"),
 				FMath::RadiansToDegrees(SimulationSpaceData.AngularVelocity.Length() * SimulationSpaceState.Dt),
 				SimulationSpaceSettings.OrientationChangeThresholdForTeleport,
-				*OwningControlRig->GetName());
+				*OwnerName.ToString());
 		}
 
 		// Note that a teleport detection shouldn't change the pose, or the current motion. We just
@@ -400,11 +392,10 @@ FRigPhysicsSimulation::FSimulationSpaceData FRigPhysicsSimulation::UpdateSimulat
 
 	if (!SimulationSpaceSettings.ExternalTurbulenceVelocity.IsNearlyZero())
 	{
-		double GlobalTime = ExecuteContext.GetAbsoluteTime();
 		FVector T(
-			FPerlinHelpers::PerlinNoise1D(GlobalTime),
-			FPerlinHelpers::PerlinNoise1D(GlobalTime + 10.0),
-			FPerlinHelpers::PerlinNoise1D(GlobalTime + 20.0));
+			FPerlinHelpers::PerlinNoise1D(AbsoluteTime),
+			FPerlinHelpers::PerlinNoise1D(AbsoluteTime + 10.0),
+			FPerlinHelpers::PerlinNoise1D(AbsoluteTime + 20.0));
 
 		FVector Turbulence = T * SimulationSpaceSettings.ExternalTurbulenceVelocity;
 
@@ -413,7 +404,7 @@ FRigPhysicsSimulation::FSimulationSpaceData FRigPhysicsSimulation::UpdateSimulat
 		if (GEngine && CVarControlRigPhysicsShowSimulationSpaceInfo.GetValueOnAnyThread())
 		{
 			GEngine->AddOnScreenDebugMessage(INDEX_NONE, 0.f, FColor::Yellow,
-				FString::Printf(TEXT("Sim %s Turbulence %s"), *OwningControlRig->GetName(), *Turbulence.ToString()));
+				FString::Printf(TEXT("Sim %s Turbulence %s"), *OwnerName.ToString(), *Turbulence.ToString()));
 		}
 	}
 
@@ -433,7 +424,7 @@ FRigPhysicsSimulation::FSimulationSpaceData FRigPhysicsSimulation::UpdateSimulat
 	{
 		GEngine->AddOnScreenDebugMessage(INDEX_NONE, 0.f, FColor::Yellow,
 			FString::Printf(TEXT("Sim %s dt %6.2fms Space P %35s V %35s A %35s"),
-				*OwningControlRig->GetName(), Dt * 1000.0f,
+				*OwnerName.ToString(), Dt * 1000.0f,
 				*SimulationSpaceState.SimulationSpaceTM.GetLocation().ToString(),
 				*SimulationSpaceData.LinearVelocity.ToString(),
 				*SimulationSpaceData.LinearAcceleration.ToString()));
diff --git a/Engine/Plugins/Experimental/ControlRigPhysics/Source/ControlRigPhysics/Private/RigPhysicsSimulationVisualization.cpp b/Engine/Plugins/Experimental/ControlRigPhysics/Source/ControlRigPhysics/Private/RigPhysicsSimulationVisualization.cpp
--- a/Engine/Plugins/Experimental/ControlRigPhysics/Source/ControlRigPhysics/Private/RigPhysicsSimulationVisualization.cpp
+++ b/Engine/Plugins/Experimental/ControlRigPhysics/Source/ControlRigPhysics/Private/RigPhysicsSimulationVisualization.cpp
@@ -4,7 +4,6 @@
 
 #include "Rigs/RigHierarchy.h"
 #include "Rigs/RigHierarchyController.h"
-#include "ControlRig.h"
 
 #include "RigVMCore/RigVMDrawInterface.h"
 
@@ -25,6 +24,7 @@
 
 #include "Math/Color.h"
 
+// Danny TODO expose these in the visualization settings
 namespace RigPhysicsSolverDraw
 {
 	FLinearColor DynamicCollisionColor = FColor::Yellow;
@@ -61,7 +61,7 @@ static FLinearColor GetColor(bool bIsKinematic, bool bCollides)
 
 //======================================================================================================================
 static void DrawShapes(
-	FRigVMDrawInterface*                    DI,
+	FRigVMDrawInterface&                    DI,
 	const FRigPhysicsVisualizationSettings& VisualizationSettings,
 	const bool                              bIsKinematic,
 	const bool                              bCollides,
@@ -103,7 +103,7 @@ static void DrawShapes(
 	case ImplicitObjectType::Sphere:
 	{
 		const Chaos::FSphere* Sphere = ImplicitObject->GetObject<Chaos::FSphere>();
-		DI->DrawSphere(ShapeTransform, FTransform(FVec3(Sphere->GetCenterOfMass())), Sphere->GetRadiusf(),
+		DI.DrawSphere(ShapeTransform, FTransform(FVec3(Sphere->GetCenterOfMass())), Sphere->GetRadiusf(),
 			GetColor(bIsKinematic, bCollides),
 			VisualizationSettings.LineThickness, VisualizationSettings.ShapeDetail);
 		break;
@@ -111,7 +111,7 @@ static void DrawShapes(
 	case ImplicitObjectType::Box:
 	{
 		const TBox<FReal, 3>* Box = ImplicitObject->GetObject<TBox<FReal, 3>>();
-		DI->DrawBox(ShapeTransform, FTransform(FQuat::Identity, Box->GetCenter(), Box->Extents()),
+		DI.DrawBox(ShapeTransform, FTransform(FQuat::Identity, Box->GetCenter(), Box->Extents()),
 			GetColor(bIsKinematic, bCollides),
 			VisualizationSettings.LineThickness);
 		break;
@@ -120,7 +120,7 @@ static void DrawShapes(
 	{
 		const Chaos::FCapsule* Capsule = ImplicitObject->GetObject<Chaos::FCapsule>();
 		const FQuat Q = FRotationMatrix::MakeFromZ(Capsule->GetAxis()).ToQuat();
-		DI->DrawCapsule(ShapeTransform, FTransform(Q, Capsule->GetCenterOfMass()),
+		DI.DrawCapsule(ShapeTransform, FTransform(Q, Capsule->GetCenterOfMass()),
 			Capsule->GetRadiusf(), Capsule->GetHeightf(),
 			GetColor(bIsKinematic, bCollides),
 			VisualizationSettings.LineThickness);
@@ -134,7 +134,7 @@ static void DrawShapes(
 
 //======================================================================================================================
 static void DrawActor(
-	FRigVMDrawInterface*                    DI,
+	FRigVMDrawInterface&                    DI,
 	const FRigPhysicsVisualizationSettings& VisualizationSettings,
 	const FTransform&                       SpaceTransform,
 	const ImmediatePhysics::FActorHandle&   ActorHandle)
@@ -173,7 +173,7 @@ float GetLimitAngleRadians(float LimitAngle, Chaos::EJointMotionType LimitType)
 
 //======================================================================================================================
 static void DrawJoint(
-	FRigVMDrawInterface*                    DI,
+	FRigVMDrawInterface&                    DI,
 	const FRigPhysicsVisualizationSettings& VisualizationSettings,
 	const FTransform&                       SpaceTransform,
 	const ImmediatePhysics::FJointHandle&   JointHandle,
@@ -204,10 +204,10 @@ static void DrawJoint(
 
 	if (bDrawAsAxes)
 	{
-		DI->DrawAxes(
+		DI.DrawAxes(
 			FTransform(JointSettings.AngularDrivePositionTarget, JointSettings.LinearDrivePositionTarget) * LimitFrame, 
 			FTransform(), Size, VisualizationSettings.LineThickness * 2);
-		DI->DrawAxes(DialFrame, FTransform(), Size, VisualizationSettings.LineThickness);
+		DI.DrawAxes(DialFrame, FTransform(), Size, VisualizationSettings.LineThickness);
 	}
 
 	{
@@ -264,7 +264,7 @@ static void DrawJoint(
 			if (JointSettings.AngularMotionTypes[Swing1Index] == Chaos::EJointMotionType::Free && 
 				JointSettings.AngularMotionTypes[Swing2Index] == Chaos::EJointMotionType::Free)
 			{
-				DI->DrawSphere(LimitFrame, FTransform(), Size * 0.2f, FColor::White, 
+				DI.DrawSphere(LimitFrame, FTransform(), Size * 0.2f, FColor::White, 
 					VisualizationSettings.LineThickness, VisualizationSettings.ShapeDetail);
 			}
 			else
@@ -273,14 +273,14 @@ static void DrawJoint(
 				ConeLimitTM.SetScale3D(FVector(Size));
 				const float Swing1LimitAngle = LimitAngleRadians[Swing1Index];
 				const float Swing2LimitAngle = LimitAngleRadians[Swing2Index];
-				DI->DrawCone(
+				DI.DrawCone(
 					ConeLimitTM, FTransform(), Swing1LimitAngle, Swing2LimitAngle, VisualizationSettings.ShapeDetail, 
-					true, FColor::Green, LimitMaterialX->GetRenderProxy(), VisualizationSettings.LineThickness);
+					true, FColor::Red, LimitMaterialX->GetRenderProxy(), VisualizationSettings.LineThickness);
 			}
 			// Draw the swing Dial indicator - shows the current orientation of the child frame
 			// relative to the parent frame on the swing axis.
 			// Start the arrow at the limit position as it can be confusing if there is joint separation
-			DI->DrawArrow(FTransform(LimitFrame.GetTranslation()), 
+			DI.DrawArrow(FTransform(LimitFrame.GetTranslation()), 
 				DialFrame.GetUnitAxis(EAxis::X) * ArrowLength,
 				DialFrame.GetUnitAxis(EAxis::Y), FColor::Red, VisualizationSettings.LineThickness);
 
@@ -288,7 +288,7 @@ static void DrawJoint(
 			{
 				// Drawing a sphere seems a bit heavy... but it appears that if DrawPoint is used then they get
 				// culled when you get close!
-				DI->DrawSphere(FTransform(
+				DI.DrawSphere(FTransform(
 					DialFrame.GetRotation(),
 					LimitFrame.GetTranslation() + DialFrame.GetUnitAxis(EAxis::X) * ArrowLength),
 					FTransform(),
@@ -304,23 +304,25 @@ static void DrawJoint(
 			FTransform ConeLimitTM = LimitFrame;
 			ConeLimitTM.SetScale3D(FVector(Size));
 			const float TwistLimitAngle = LimitAngleRadians[TwistIndex];
-			DI->DrawCone(ConeLimitTM, FTransform(FQuat::MakeFromRotationVector(FVector(0, 1, 0) * -HALF_PI)),
+			DI.DrawCone(ConeLimitTM, FTransform(
+				FQuat::MakeFromRotationVector(FVector(0, 1, 0) * -HALF_PI) * 
+				FQuat::MakeFromRotationVector(FVector(0, 0, 1) * HALF_PI)),
 				TwistLimitAngle, 0.0f, VisualizationSettings.ShapeDetail, true, FColor::Green,
 				LimitMaterialYAxis->GetRenderProxy(), VisualizationSettings.LineThickness);
 
 			// Draw the twist Dial indicator - shows the current orientation of the child frame
 			// relative to the parent frame on the twist axis.
 			FQuat Rot(LimitFrame.GetUnitAxis(EAxis::X), TwistAngle);
-			FVector TwistArrow = Rot * LimitFrame.GetUnitAxis(EAxis::Z);
-			FVector TwistSideDir = FVector::CrossProduct(LimitFrame.GetUnitAxis(EAxis::X), TwistArrow);
+			FVector TwistArrow = Rot * LimitFrame.GetUnitAxis(EAxis::Y);
+			FVector TwistSideDir = FVector::CrossProduct(LimitFrame.GetUnitAxis(EAxis::Z), TwistArrow);
 
-			DI->DrawArrow(FTransform(LimitFrame.GetTranslation()),
+			DI.DrawArrow(FTransform(LimitFrame.GetTranslation()),
 				TwistArrow * ArrowLength,
-				TwistSideDir, FColor::Blue, VisualizationSettings.LineThickness);
+				TwistSideDir, FColor::Green, VisualizationSettings.LineThickness);
 
 			if (bTwistViolated)
 			{
-				DI->DrawSphere(FTransform(
+				DI.DrawSphere(FTransform(
 					DialFrame.GetRotation(),
 					LimitFrame.GetTranslation() + TwistArrow * ArrowLength),
 					FTransform(),
@@ -331,9 +333,9 @@ static void DrawJoint(
 }
 
 //======================================================================================================================
-static bool IsConstraintDriveActive(const Chaos::FPBDJointConstraintHandle* Constraint)
+static bool IsConstraintDriveActive(const Chaos::FPBDJointConstraintHandle& Constraint)
 {
-	const Chaos::FPBDJointSettings& Settings = Constraint->GetSettings();
+	const Chaos::FPBDJointSettings& Settings = Constraint.GetSettings();
 	if (Settings.bAngularSLerpPositionDriveEnabled || Settings.bAngularSwingPositionDriveEnabled || Settings.bAngularTwistPositionDriveEnabled)
 	{
 		if (Settings.AngularDriveStiffness.SizeSquared() > 0.0)
@@ -378,8 +380,7 @@ void FRigPhysicsSimulation::Draw(
 	}
 #endif
 
-	const URigHierarchy* Hierarchy = OwningControlRig->GetHierarchy();
-	if (!DI || !Hierarchy)
+	if (!DI)
 	{
 		return;
 	}
@@ -392,7 +393,7 @@ void FRigPhysicsSimulation::Draw(
 	{
 		if (CollisionActorHandle)
 		{
-			DrawActor(DI, VisualizationSettings, SpaceTransform, *CollisionActorHandle);
+			DrawActor(*DI, VisualizationSettings, SpaceTransform, *CollisionActorHandle);
 		}
 
 		for (const TPair<FRigComponentKey, FRigBodyRecord>& BodyRecordPair : BodyRecords)
@@ -400,7 +401,7 @@ void FRigPhysicsSimulation::Draw(
 			const FRigBodyRecord& Record = BodyRecordPair.Value;
 			if (const ImmediatePhysics::FActorHandle* ActorHandle = Record.ActorHandle)
 			{
-				DrawActor(DI, VisualizationSettings, SpaceTransform, *ActorHandle);
+				DrawActor(*DI, VisualizationSettings, SpaceTransform, *ActorHandle);
 			}
 		}
 	}
@@ -412,7 +413,7 @@ void FRigPhysicsSimulation::Draw(
 			const FRigJointRecord& Record = JointRecordPair.Value;
 			if (const ImmediatePhysics::FJointHandle* JointHandle = Record.JointHandle)
 			{
-				DrawJoint(DI, VisualizationSettings, SpaceTransform, *JointHandle, false);
+				DrawJoint(*DI, VisualizationSettings, SpaceTransform, *JointHandle, false);
 			}
 		}
 	}
@@ -426,9 +427,9 @@ void FRigPhysicsSimulation::Draw(
 			{
 				if (const Chaos::FPBDJointConstraintHandle* Constraint = JointHandle->GetConstraint())
 				{
-					if (IsConstraintDriveActive(Constraint))
+					if (IsConstraintDriveActive(*Constraint))
 					{
-						DrawJoint(DI, VisualizationSettings, SpaceTransform, *JointHandle, true);
+						DrawJoint(*DI, VisualizationSettings, SpaceTransform, *JointHandle, true);
 					}
 				}
 			}
diff --git a/Engine/Plugins/Experimental/ControlRigPhysics/Source/ControlRigPhysics/Public/RigPhysicsBodyComponent.h b/Engine/Plugins/Experimental/ControlRigPhysics/Source/ControlRigPhysics/Public/RigPhysicsBodyComponent.h
--- a/Engine/Plugins/Experimental/ControlRigPhysics/Source/ControlRigPhysics/Public/RigPhysicsBodyComponent.h
+++ b/Engine/Plugins/Experimental/ControlRigPhysics/Source/ControlRigPhysics/Public/RigPhysicsBodyComponent.h
@@ -60,9 +60,6 @@ public:
 	UE_API virtual const FSlateIcon& GetIconForUI() const override;
 #endif
 
-	UE_API virtual bool CanBeAddedTo(
-		const FRigElementKey& InElementKey, const URigHierarchy* InHierarchy, FString* OutFailureReason) const override;
-
 	UE_API virtual void OnAddedToHierarchy(URigHierarchy* InHierarchy, URigHierarchyController* InController) override;
 
 	UE_API virtual void OnRigHierarchyKeyChanged(const FRigHierarchyKey& InOldKey, const FRigHierarchyKey& InNewKey) override;
diff --git a/Engine/Plugins/Experimental/ControlRigPhysics/Source/ControlRigPhysics/Public/RigPhysicsBodyExecution.h b/Engine/Plugins/Experimental/ControlRigPhysics/Source/ControlRigPhysics/Public/RigPhysicsBodyExecution.h
--- a/Engine/Plugins/Experimental/ControlRigPhysics/Source/ControlRigPhysics/Public/RigPhysicsBodyExecution.h
+++ b/Engine/Plugins/Experimental/ControlRigPhysics/Source/ControlRigPhysics/Public/RigPhysicsBodyExecution.h
@@ -18,7 +18,7 @@
  * Adds a new physics body as a component on the owner element.
  * Note: This node only runs as part of the construction event.
  */
-USTRUCT(meta = (DisplayName = "Add Physics Body Component", Keywords = "Spawn,Construction,Create,New,Body,Skeleton", Varying))
+USTRUCT(meta = (DisplayName = "Spawn Physics Body", Keywords = "Add,Construction,Create,New,Body,Skeleton", Varying))
 struct FRigUnit_AddPhysicsBody : public FRigUnit_PhysicsBaseMutable
 {
 	GENERATED_BODY()
@@ -27,7 +27,7 @@ struct FRigUnit_AddPhysicsBody : public FRigUnit_PhysicsBaseMutable
 	{
 		Owner.Type = ERigElementType::Bone;
 		Solver.PhysicsSolverComponentKey.ElementKey.Type = ERigElementType::Bone;
-		Solver.PhysicsSolverComponentKey.Name = FRigPhysicsBodyComponent::GetDefaultName();
+		Solver.PhysicsSolverComponentKey.Name = FRigPhysicsSolverComponent::GetDefaultName();
 	}
 
 	RIGVM_METHOD()
@@ -63,7 +63,7 @@ struct FRigUnit_AddPhysicsBody : public FRigUnit_PhysicsBaseMutable
 
 // Discards any existing collision data and replaces it with a box based on the joint positions.
 // Note that this must be called before the physics solver is instantiated/stepped.
-USTRUCT(meta = (DisplayName = "Automatically Calculate Collision", Varying))
+USTRUCT(meta = (DisplayName = "Calculate Physics Collision", Varying))
 struct FRigUnit_HierarchyAutoCalculateCollision : public FRigUnit_PhysicsBaseMutable
 {
 	GENERATED_BODY()
diff --git a/Engine/Plugins/Experimental/ControlRigPhysics/Source/ControlRigPhysics/Public/RigPhysicsControlExecution.h b/Engine/Plugins/Experimental/ControlRigPhysics/Source/ControlRigPhysics/Public/RigPhysicsControlExecution.h
--- a/Engine/Plugins/Experimental/ControlRigPhysics/Source/ControlRigPhysics/Public/RigPhysicsControlExecution.h
+++ b/Engine/Plugins/Experimental/ControlRigPhysics/Source/ControlRigPhysics/Public/RigPhysicsControlExecution.h
@@ -16,11 +16,20 @@
  * Adds a new physics control as a component on the owner element.
  * Note: This node only runs as part of the construction event.
  */
-USTRUCT(meta = (Category = "RigPhysics", NodeColor = "1.0 0.6 0.3", DisplayName = "Add Physics Control Component", Keywords = "Spawn,Construction,Create,New,Joint,Control,Physics", Varying))
+USTRUCT(meta = (Category = "RigPhysics", NodeColor = "1.0 0.6 0.3", DisplayName = "Spawn Physics Control", Keywords = "Add,Construction,Create,New,Control", Varying))
 struct FRigUnit_AddPhysicsControl : public FRigUnit_PhysicsBaseMutable
 {
 	GENERATED_BODY()
 
+	FRigUnit_AddPhysicsControl()
+	{
+		Owner.Type = ERigElementType::Bone;
+		ParentBodyComponentKey.ElementKey.Type = ERigElementType::Bone;
+		ParentBodyComponentKey.Name = FRigPhysicsBodyComponent::GetDefaultName();
+		ChildBodyComponentKey.ElementKey.Type = ERigElementType::Bone;
+		ChildBodyComponentKey.Name = FRigPhysicsBodyComponent::GetDefaultName();
+	}
+
 	RIGVM_METHOD()
 	UE_API virtual void Execute() override;
 
@@ -143,7 +152,7 @@ struct FRigUnit_HierarchySetControlMultiplier : public FRigUnit_PhysicsBaseMutab
 	}
 
 	RIGVM_METHOD()
-		UE_API virtual void Execute() override;
+	UE_API virtual void Execute() override;
 
 	UPROPERTY(meta = (Input))
 	FRigComponentKey PhysicsControlComponentKey;
@@ -165,7 +174,7 @@ struct FRigUnit_HierarchySetControlDataAndMultiplier : public FRigUnit_PhysicsBa
 	}
 
 	RIGVM_METHOD()
-		UE_API virtual void Execute() override;
+	UE_API virtual void Execute() override;
 
 	UPROPERTY(meta = (Input))
 	FRigComponentKey PhysicsControlComponentKey;
diff --git a/Engine/Plugins/Experimental/ControlRigPhysics/Source/ControlRigPhysics/Public/RigPhysicsExecution.h b/Engine/Plugins/Experimental/ControlRigPhysics/Source/ControlRigPhysics/Public/RigPhysicsExecution.h
--- a/Engine/Plugins/Experimental/ControlRigPhysics/Source/ControlRigPhysics/Public/RigPhysicsExecution.h
+++ b/Engine/Plugins/Experimental/ControlRigPhysics/Source/ControlRigPhysics/Public/RigPhysicsExecution.h
@@ -35,7 +35,7 @@ struct FRigUnit_PhysicsBase : public FRigUnit
  * Adds a new physics solver as a component on the owner element.
  * Note: This node only runs as part of the construction event.
  */
-USTRUCT(meta=(DisplayName="Add Physics Solver Component", Keywords="Spawn,Construction,Create,New,Simulation", Varying))
+USTRUCT(meta=(DisplayName="Spawn Physics Solver", Keywords="Add,Construction,Create,New,Simulation", Varying))
 struct FRigUnit_AddPhysicsSolver : public FRigUnit_PhysicsBaseMutable
 {
 	GENERATED_BODY()
@@ -205,7 +205,7 @@ struct FRigUnit_SetPhysicsSolverExternalVelocity : public FRigUnit_PhysicsBaseMu
 /**
  * Forces tracking of the input animation (on all physics bodies) for the next N frames
  */
-USTRUCT(meta = (DisplayName = "Track Input Pose", Keywords = "Simulate"))
+USTRUCT(meta = (DisplayName = "Track Input Pose", Keywords = "Reset,Simulate"))
 struct FRigUnit_TrackInputPose : public FRigUnit_PhysicsBaseMutable
 {
 	GENERATED_BODY()
@@ -235,7 +235,7 @@ struct FRigUnit_TrackInputPose : public FRigUnit_PhysicsBaseMutable
 
 
 // Adds a set of physics components including the body, joint and controls
-USTRUCT(meta = (DisplayName = "Add Physics Components", Keywords = "Spawn,Construction,Create,New", Varying))
+USTRUCT(meta = (DisplayName = "Spawn Physics Components", Keywords = "Add,Construction,Create,New,Body,Joint,Control", Varying))
 struct FRigUnit_AddPhysicsComponents : public FRigUnit_PhysicsBaseMutable
 {
 	GENERATED_BODY()
diff --git a/Engine/Plugins/Experimental/ControlRigPhysics/Source/ControlRigPhysics/Public/RigPhysicsJointExecution.h b/Engine/Plugins/Experimental/ControlRigPhysics/Source/ControlRigPhysics/Public/RigPhysicsJointExecution.h
--- a/Engine/Plugins/Experimental/ControlRigPhysics/Source/ControlRigPhysics/Public/RigPhysicsJointExecution.h
+++ b/Engine/Plugins/Experimental/ControlRigPhysics/Source/ControlRigPhysics/Public/RigPhysicsJointExecution.h
@@ -18,7 +18,7 @@
  * Adds a new physics body as a component on the owner element.
  * Note: This node only runs as part of the construction event.
  */
-USTRUCT(meta = (DisplayName = "Add Physics Joint Component", Keywords = "Spawn,Construction,Create,New,Joint,Articulation,Skeleton", Varying))
+USTRUCT(meta = (DisplayName = "Spawn Physics Joint", Keywords = "Add,Construction,Create,New,Joint,Articulation,Constraint,Skeleton", Varying))
 struct FRigUnit_AddPhysicsJoint : public FRigUnit_PhysicsBaseMutable
 {
 	GENERATED_BODY()
@@ -26,6 +26,10 @@ struct FRigUnit_AddPhysicsJoint : public FRigUnit_PhysicsBaseMutable
 	FRigUnit_AddPhysicsJoint()
 	{
 		Owner.Type = ERigElementType::Bone;
+		ParentBodyComponentKey.ElementKey.Type = ERigElementType::Bone;
+		ParentBodyComponentKey.Name = FRigPhysicsBodyComponent::GetDefaultName();
+		ChildBodyComponentKey.ElementKey.Type = ERigElementType::Bone;
+		ChildBodyComponentKey.Name = FRigPhysicsBodyComponent::GetDefaultName();
 	}
 
 	RIGVM_METHOD()
@@ -310,5 +314,4 @@ struct FRigUnit_MakeDriveData : public FRigUnit_PhysicsBase
 	FRigPhysicsDriveData DriveData;
 };
 
-
 #undef UE_API
