Control Rig Physics support for colliding against world objects

diff --git a/Engine/Plugins/Experimental/ControlRigPhysics/Source/ControlRigPhysics/Private/RigPhysicsData.cpp b/Engine/Plugins/Experimental/ControlRigPhysics/Source/ControlRigPhysics/Private/RigPhysicsData.cpp
--- a/Engine/Plugins/Experimental/ControlRigPhysics/Source/ControlRigPhysics/Private/RigPhysicsData.cpp
+++ b/Engine/Plugins/Experimental/ControlRigPhysics/Source/ControlRigPhysics/Private/RigPhysicsData.cpp
@@ -268,6 +268,12 @@ FArchive& operator <<(FArchive& Ar, FRigPhysicsSolverSettings& Data)
 	{
 		Ar << Data.bAutomaticallyAddPhysicsComponents;
 	}
+	if (Ar.CustomVer(FPhysicsControlObjectVersion::GUID) >= FPhysicsControlObjectVersion::ControlRigWorldCollision)
+	{
+		Ar << Data.WorldCollisionType;
+		Ar << Data.WorldCollisionExpiryFrames;
+		Ar << Data.WorldCollisionBoundsExpansion;
+	}
 	return Ar;
 }
 
diff --git a/Engine/Plugins/Experimental/ControlRigPhysics/Source/ControlRigPhysics/Private/RigPhysicsExecution.cpp b/Engine/Plugins/Experimental/ControlRigPhysics/Source/ControlRigPhysics/Private/RigPhysicsExecution.cpp
--- a/Engine/Plugins/Experimental/ControlRigPhysics/Source/ControlRigPhysics/Private/RigPhysicsExecution.cpp
+++ b/Engine/Plugins/Experimental/ControlRigPhysics/Source/ControlRigPhysics/Private/RigPhysicsExecution.cpp
@@ -123,10 +123,15 @@ FRigUnit_StepPhysicsSolver_Execute()
 
 			if (PhysicsSimulation)
 			{
+				const USceneComponent* OwningSceneComponent = ControlRig->GetOwningSceneComponent();
+				const AActor* OwningActorPtr = OwningSceneComponent ? OwningSceneComponent->GetOwner() : nullptr;
+
 				PhysicsSimulation->StepSimulation(
-					ExecuteContext, Hierarchy, SolverComponent, DeltaTimeOverride, SimulationSpaceDeltaTimeOverride, Alpha);
+					ControlRig->GetWorld(), OwningActorPtr, ExecuteContext, Hierarchy, SolverComponent, 
+					DeltaTimeOverride, SimulationSpaceDeltaTimeOverride, Alpha);
 				int32 VisualizationOverride = CVarControlRigPhysicsVisualizationOverride.GetValueOnAnyThread();
-				if (VisualizationOverride == 1 || (VisualizationOverride < 0 && VisualizationSettings.bEnableVisualization))
+				if (VisualizationOverride == 1 || 
+					(VisualizationOverride < 0 && VisualizationSettings.bEnableVisualization))
 				{
 					PhysicsSimulation->Draw(
 						ExecuteContext.GetDrawInterface(), SolverComponent->SolverSettings,
diff --git a/Engine/Plugins/Experimental/ControlRigPhysics/Source/ControlRigPhysics/Private/RigPhysicsSimulation.cpp b/Engine/Plugins/Experimental/ControlRigPhysics/Source/ControlRigPhysics/Private/RigPhysicsSimulation.cpp
--- a/Engine/Plugins/Experimental/ControlRigPhysics/Source/ControlRigPhysics/Private/RigPhysicsSimulation.cpp
+++ b/Engine/Plugins/Experimental/ControlRigPhysics/Source/ControlRigPhysics/Private/RigPhysicsSimulation.cpp
@@ -67,6 +67,8 @@ constexpr int32 ConstraintParentIndex = 1;
 FRigPhysicsSimulation::FRigPhysicsSimulation(const FName InOwnerName)
 	: FRigPhysicsSimulationBase(FRigPhysicsSimulation::StaticStruct()), OwnerName(InOwnerName)
 {
+	WorldObjects = MakeShared<TMap<uint32, FWorldObject>>();
+	SimulationMutex = MakeShared<FCriticalSection>();
 }
 
 //======================================================================================================================
@@ -331,7 +333,6 @@ void FRigPhysicsSimulation::DestroyPhysicsSimulation()
 	}
 	BodyRecords.Empty();
 
-
 	for (TPair<FRigComponentKey, FRigJointRecord>& JointRecordPair : JointRecords)
 	{
 		FRigJointRecord& Record = JointRecordPair.Value;
@@ -584,7 +585,7 @@ ImmediatePhysics::FActorHandle* FRigPhysicsSimulation::CreateBody(
 	if (!ActorHandle)
 	{
 		UE_LOG(LogRigPhysics, Warning,
-			TEXT("Unable to create body %s"), *BodyName.ToString());
+			TEXT("Control Rig %s Unable to create body %s"), *OwnerName.ToString(), *BodyName.ToString());
 		return nullptr;
 	}
 
@@ -628,6 +629,9 @@ void FRigPhysicsSimulation::Instantiate(
 		return;
 	}
 
+	// Lock against access to the simulation, in case of access by the WorldObject task
+	FScopeLock Lock(SimulationMutex.Get());
+
 	// We need the simulation space in order to instantiate properly. This is not ideal, as we may
 	// end up updating the simulation space data twice (thus inserting data into the history). This
 	// shouldn't really matter as it will only be on the first step.
@@ -891,7 +895,8 @@ void FRigPhysicsSimulation::InstantiateControls(
 				if (!ControlRecord.JointHandle)
 				{
 					UE_LOG(LogRigPhysics, Warning,
-						TEXT("Unable to create control constraint for %s"), *ControlComponent->GetKey().ToString());
+						TEXT("Control Rig %s Unable to create control constraint for %s"), 
+						*OwnerName.ToString(), *ControlComponent->GetKey().ToString());
 				}
 
 				if (ControlComponent->ControlData.bDisableCollision)
@@ -1608,12 +1613,14 @@ void FRigPhysicsSimulation::UpdatePostPhysics(
 
 //======================================================================================================================
 void FRigPhysicsSimulation::StepSimulation(
-	const FRigVMExecuteContext& ExecuteContext,
-	URigHierarchy&              Hierarchy,
-	FRigPhysicsSolverComponent* SolverComponent,
-	const float                 DeltaTimeOverride,
-	const float                 SimulationSpaceDeltaTimeOverride,
-	const float                 Alpha)
+	const UWorld*                World,
+	const AActor*                OwningActorPtr,
+	const FRigVMExecuteContext&  ExecuteContext,
+	URigHierarchy&               Hierarchy,
+	FRigPhysicsSolverComponent*  SolverComponent,
+	const float                  DeltaTimeOverride,
+	const float                  SimulationSpaceDeltaTimeOverride,
+	const float                  Alpha)
 {
 	DECLARE_SCOPE_HIERARCHICAL_COUNTER_FUNC()
 	QUICK_SCOPE_CYCLE_COUNTER(STAT_RigPhysics_StepSimulation);
@@ -1637,6 +1644,9 @@ void FRigPhysicsSimulation::StepSimulation(
 		PhysicsSimulationSpaceDeltaTime = SimulationSpaceDeltaTimeOverride;
 	}
 
+	// Lock against access to the simulation, in case of access by the WorldObject task
+	FScopeLock Lock(SimulationMutex.Get());
+
 	// We need to know about the simulation space etc before we can instantiate anything into the right place
 	SimulationSpaceData = UpdateSimulationSpaceStateAndCalculateData(
 		ExecuteContext, Hierarchy, SolverComponent, PhysicsSimulationSpaceDeltaTime, ExecuteContext.GetAbsoluteTime());
@@ -1701,6 +1711,8 @@ void FRigPhysicsSimulation::StepSimulation(
 		SimulationSpaceSettings.SpaceMovementAmount * SimulationSpaceData.LinearAcceleration,
 		SimulationSpaceSettings.SpaceMovementAmount * SimulationSpaceData.AngularAcceleration);
 
+	UpdateWorldObjectsPrePhysics(SolverSettings);
+
 	// Only update if there is a delta time:
 	// * We don't want to update our previous TMs and store the dt - because that would end up
 	//   implying infinite velocities
@@ -1721,6 +1733,13 @@ void FRigPhysicsSimulation::StepSimulation(
 	// Always do a read-back, even for zero Dt
 	UpdatePostPhysics(Hierarchy, SolverComponent, FMath::Clamp(Alpha, 0.0f, 1.0f), PhysicsDeltaTime, FixedTimeStep > 0);
 
+	// Trigger the update of world objects as soon as possible. Note that this needs to be done
+	// after the last access to the simulation, because the game-thread task will start
+	// adding/removing objects from our simulation, and we can't let that work overlap with what we
+	// do here.
+
+	UpdateWorldObjectsPostPhysics(World, SolverSettings, OwningActorPtr);
+
 	if (SolverComponent->TrackInputCounter > 0)
 	{
 		--SolverComponent->TrackInputCounter;
diff --git a/Engine/Plugins/Experimental/ControlRigPhysics/Source/ControlRigPhysics/Private/RigPhysicsSimulation.h b/Engine/Plugins/Experimental/ControlRigPhysics/Source/ControlRigPhysics/Private/RigPhysicsSimulation.h
--- a/Engine/Plugins/Experimental/ControlRigPhysics/Source/ControlRigPhysics/Private/RigPhysicsSimulation.h
+++ b/Engine/Plugins/Experimental/ControlRigPhysics/Source/ControlRigPhysics/Private/RigPhysicsSimulation.h
@@ -12,6 +12,7 @@
 #include "Physics/ImmediatePhysics/ImmediatePhysicsDeclares.h"
 #include "Physics/ImmediatePhysics/ImmediatePhysicsSimulation.h"
 #include "Chaos/PBDJointConstraintTypes.h"
+#include "HAL/CriticalSection.h"
 
 #include "RigPhysicsSimulation.generated.h"
 
@@ -155,12 +156,14 @@ public:
 	// If it is zero, then delta time is taken from the execute context
 	// If it is negative, then the simulation isn't stepped.
 	void StepSimulation(
-		const FRigVMExecuteContext& ExecuteContext, 
-		URigHierarchy&              Hierarchy,
-		FRigPhysicsSolverComponent* SolverComponent,
-		const float                 DeltaTimeOverride, 
-		const float                 SimulationSpaceDeltaTimeOverride,
-		const float                 Alpha);
+		const UWorld*                World,
+		const AActor*                OwningActorPtr,
+		const FRigVMExecuteContext&  ExecuteContext, 
+		URigHierarchy&               Hierarchy,
+		FRigPhysicsSolverComponent*  SolverComponent,
+		const float                  DeltaTimeOverride, 
+		const float                  SimulationSpaceDeltaTimeOverride,
+		const float                  Alpha);
 
 	// Draws shapes etc AND (potentially) enables the low-level Chaos debug draw
 	void Draw(
@@ -294,6 +297,15 @@ private:
 		const float                 DeltaTime,
 		const bool                  bUsingFixedStep);
 
+	// This will walk through the WorldObjects, which should now be updated, and update the simulation
+	void UpdateWorldObjectsPrePhysics(const FRigPhysicsSolverSettings& SolverSettings);
+
+	// This will trigger a game-thread update of WorldObjects, ready for the next update
+	void UpdateWorldObjectsPostPhysics(
+		const UWorld*                    World,
+		const FRigPhysicsSolverSettings& SolverSetting, 
+		const AActor*                    OwningActorPtr);
+
 	void InitSimulationSpace(const FTransform& ComponentToWorld, const FTransform& BoneToComponent);
 
 	void CheckForResetsPrePhysics(
@@ -327,6 +339,10 @@ private:
 	FVector ConvertWorldVectorToSimSpaceNoScale(
 		const FRigPhysicsSolverSettings& SolverSettings, const FVector& WorldVector) const;
 
+	// Converts a transform specified in world space to the simulation space
+	FTransform ConvertWorldTransformToSimSpace(
+		const FRigPhysicsSolverSettings& SolverSettings, const FTransform& WorldTM) const;
+
 	FTransform ConvertCollisionSpaceTransformToSimSpace(
 		const FRigPhysicsSolverSettings& SolverSettings, const FTransform& TM) const;
 
@@ -385,5 +401,55 @@ private:
 	// This will be set to false after instantiation, which can happen manually during construction,
 	// or be postponed until there is a simulation step.
 	bool bNeedToInstantiate = true;
+
+	// Record of a world object that we want to track, and "import" into our simulation for collision.
+	struct FWorldObject
+	{
+		FWorldObject() : ActorHandle(nullptr), LastSeenUpdateCounter(0) {}
+		FWorldObject(ImmediatePhysics::FActorHandle* InActorHandle, int64 InLastSeenUpdateCounter)
+			: ActorHandle(InActorHandle)
+			, LastSeenUpdateCounter(InLastSeenUpdateCounter) {
+		}
+
+		// The object we are tracking
+		TWeakObjectPtr<UPrimitiveComponent> WorldPrimitiveComponent;
+
+		// The actor in our simulation
+		ImmediatePhysics::FActorHandle* ActorHandle;
+
+		// The transform of the world object. This will be updated in the game-thread task.
+		FTransform ComponentWorldTransform;
+
+		// When the world object was last seen (i.e. captured by an overlap), used to detect when we
+		// should expire/forget about it.
+		int64 LastSeenUpdateCounter;
+
+		bool GetExpired(int64 InUpdateCounter, int64 InRetainCount) const
+		{
+			if (LastSeenUpdateCounter != -1 && InUpdateCounter > LastSeenUpdateCounter + InRetainCount)
+			{
+				return true;
+			}
+			if (!WorldPrimitiveComponent.IsValid() || !ActorHandle)
+			{
+				return true;
+			}
+			return false;
+		}
+	};
+
+	// The world objects we track for collision. Index uses GetUniqueID from the PrimitiveComponent.
+	// Note that this is only guaranteed to be unique whilst the object exists, so it is possible
+	// (but hopefully unlikely) that an object will be replaced and have the same ID. This will be
+	// rectified on the next update, as the record's pointer to the primitive component will be invalid.
+	TSharedPtr<TMap<uint32, FWorldObject>> WorldObjects;
+
+	// The box used for overlap tests to find the world objects
+	FBox WorldOverlapBox;
+
+	// Guard against access to the low level simulation. Note that this access may come from an
+	// async/spawned task, so the mutex itself needs to be shared, as we (as primary owners) may get
+	// deleted whilst the spawned task is in flight.
+	TSharedPtr<FCriticalSection> SimulationMutex;
 };
 
diff --git a/Engine/Plugins/Experimental/ControlRigPhysics/Source/ControlRigPhysics/Private/RigPhysicsSimulationSpace.cpp b/Engine/Plugins/Experimental/ControlRigPhysics/Source/ControlRigPhysics/Private/RigPhysicsSimulationSpace.cpp
--- a/Engine/Plugins/Experimental/ControlRigPhysics/Source/ControlRigPhysics/Private/RigPhysicsSimulationSpace.cpp
+++ b/Engine/Plugins/Experimental/ControlRigPhysics/Source/ControlRigPhysics/Private/RigPhysicsSimulationSpace.cpp
@@ -124,6 +124,20 @@ FVector FRigPhysicsSimulation::ConvertWorldVectorToSimSpaceNoScale(
 	}
 }
 
+//======================================================================================================================
+FTransform FRigPhysicsSimulation::ConvertWorldTransformToSimSpace(
+	const FRigPhysicsSolverSettings& SolverSettings, const FTransform& WorldTM) const
+{
+	switch (SolverSettings.SimulationSpace)
+	{
+	case ERigPhysicsSimulationSpace::Component: return WorldTM.GetRelativeTransform(SimulationSpaceState.ComponentTM);
+	case ERigPhysicsSimulationSpace::World: return WorldTM;
+	case ERigPhysicsSimulationSpace::SpaceBone:
+		return WorldTM.GetRelativeTransform(SimulationSpaceState.BoneRelComponentTM * SimulationSpaceState.ComponentTM);
+	default: ensureMsgf(false, TEXT("Unsupported Simulation Space")); return WorldTM;
+	}
+}
+
 //======================================================================================================================
 FTransform FRigPhysicsSimulation::ConvertCollisionSpaceTransformToSimSpace(
 	const FRigPhysicsSolverSettings& SolverSettings, const FTransform& TM) const
diff --git a/Engine/Plugins/Experimental/ControlRigPhysics/Source/ControlRigPhysics/Private/RigPhysicsSimulationVisualization.cpp b/Engine/Plugins/Experimental/ControlRigPhysics/Source/ControlRigPhysics/Private/RigPhysicsSimulationVisualization.cpp
--- a/Engine/Plugins/Experimental/ControlRigPhysics/Source/ControlRigPhysics/Private/RigPhysicsSimulationVisualization.cpp
+++ b/Engine/Plugins/Experimental/ControlRigPhysics/Source/ControlRigPhysics/Private/RigPhysicsSimulationVisualization.cpp
@@ -110,8 +110,10 @@ static void DrawShapes(
 	}
 	case ImplicitObjectType::Box:
 	{
-		const TBox<FReal, 3>* Box = ImplicitObject->GetObject<TBox<FReal, 3>>();
-		DI.DrawBox(ShapeTransform, FTransform(FQuat::Identity, Box->GetCenter(), Box->Extents()),
+		const TBox<FReal, 3>* Box = ImplicitObject->GetObject<TBox<FReal, 3>>(); // Axis aligned
+		DI.DrawBox(
+			FTransform(FQuat::Identity, Box->GetCenter()) * ShapeTransform, 
+			FTransform(FQuat::Identity, FVector::ZeroVector, Box->Extents()),
 			GetColor(bIsKinematic, bCollides),
 			VisualizationSettings.LineThickness);
 		break;
@@ -371,6 +373,9 @@ void FRigPhysicsSimulation::Draw(
 	DECLARE_SCOPE_HIERARCHICAL_COUNTER_FUNC()
 	QUICK_SCOPE_CYCLE_COUNTER(STAT_RigPhysics_Draw);
 
+	// Lock against access to the simulation, in case of access by the WorldObject task
+	FScopeLock Lock(SimulationMutex.Get());
+
 #if CHAOS_DEBUG_DRAW
 	// Danny TODO enable/disable chaos debug draw with a cvar or flag
 	if (DebugWorld && DebugWorld->GetPhysicsScene() && DebugWorld->GetPhysicsScene()->GetDebugDrawScene())
@@ -436,6 +441,33 @@ void FRigPhysicsSimulation::Draw(
 		}
 	}
 
+	if (VisualizationSettings.bShowWorldOverlapBox)
+	{
+		if (WorldOverlapBox.IsValid)
+		{
+			// Remember that rendering is done in component space
+			FTransform BoxRelWorldTM(FQuat::Identity, WorldOverlapBox.GetCenter());
+			FTransform BoxRelComponentTM = BoxRelWorldTM.GetRelativeTransform(SimulationSpaceState.ComponentTM);
+			DI->DrawBox(
+				BoxRelComponentTM,
+				FTransform(FQuat::Identity, FVector::ZeroVector, WorldOverlapBox.GetSize()),
+				FColor::Green, 
+				VisualizationSettings.LineThickness);
+		}
+	}
+
+	if (VisualizationSettings.bShowWorldObjects)
+	{
+		for (TPair<uint32, FWorldObject>& WorldObjectPair : *WorldObjects.Get())
+		{
+			FWorldObject& WorldObject = WorldObjectPair.Value;
+			if (WorldObject.ActorHandle)
+			{
+				DrawActor(*DI, VisualizationSettings, SpaceTransform, *WorldObject.ActorHandle);
+			}
+		}
+	}
+
 	int32 ShowActiveOverride = CVarControlRigPhysicsShowActiveContactsOverride.GetValueOnAnyThread();
 	int32 ShowInactiveOverride = CVarControlRigPhysicsShowInactiveContactsOverride.GetValueOnAnyThread();
 
diff --git a/Engine/Plugins/Experimental/ControlRigPhysics/Source/ControlRigPhysics/Private/RigPhysicsSimulationWorldObjects.cpp b/Engine/Plugins/Experimental/ControlRigPhysics/Source/ControlRigPhysics/Private/RigPhysicsSimulationWorldObjects.cpp
new file mode 100644
--- /dev/null
+++ b/Engine/Plugins/Experimental/ControlRigPhysics/Source/ControlRigPhysics/Private/RigPhysicsSimulationWorldObjects.cpp
@@ -0,0 +1,233 @@
+// Copyright Epic Games, Inc. All Rights Reserved.
+
+#include "RigPhysicsSimulation.h"
+
+#include "Physics/ImmediatePhysics/ImmediatePhysicsActorHandle.h"
+#include "Physics/ImmediatePhysics/ImmediatePhysicsAdapters.h"
+#include "Physics/ImmediatePhysics/ImmediatePhysicsSimulation.h"
+#include "Physics/ImmediatePhysics/ImmediatePhysicsChaos/ImmediatePhysicsSimulation_Chaos.h"
+#include "Chaos/ParticleHandle.h"
+
+#include "Engine/OverlapResult.h"
+#include "Engine/World.h"
+#include "CollisionQueryParams.h"
+#include "CollisionShape.h"
+
+//======================================================================================================================
+// Note that the simulation mutex will already be taken when we are called
+void FRigPhysicsSimulation::UpdateWorldObjectsPrePhysics(const FRigPhysicsSolverSettings& SolverSettings)
+{
+	if (SolverSettings.WorldCollisionType == ERigPhysicsWorldCollisionType::None)
+	{
+		for (TPair<uint32, FWorldObject>& WorldObjectPair : *WorldObjects.Get())
+		{
+			FWorldObject& WorldObject = WorldObjectPair.Value;
+			if (WorldObject.ActorHandle)
+			{
+				UE_LOG(LogRigPhysics, Warning,
+					TEXT("Control rig %s - expiring object %s"),
+					*OwnerName.ToString(), *WorldObject.ActorHandle->GetName().ToString());
+
+				Simulation->DestroyActor(WorldObject.ActorHandle);
+			}
+		}
+		WorldObjects->Reset();
+	}
+	else
+	{
+		TArray<uint32> EntriesToRemove;
+		for (TPair<uint32, FWorldObject>& WorldObjectPair : *WorldObjects.Get())
+		{
+			FWorldObject& WorldObject = WorldObjectPair.Value;
+			if (WorldObject.GetExpired(UpdateCounter, SolverSettings.WorldCollisionExpiryFrames))
+			{
+				if (WorldObject.ActorHandle)
+				{
+					UE_LOG(LogRigPhysics, Log,
+						TEXT("Control rig %s - Expiring object %s"), 
+						*OwnerName.ToString(), *WorldObject.ActorHandle->GetName().ToString());
+
+					Simulation->DestroyActor(WorldObject.ActorHandle);
+					WorldObject.ActorHandle = nullptr;
+				}
+				EntriesToRemove.Add(WorldObjectPair.Key);
+			}
+		}
+		for (uint32 EntryToRemove : EntriesToRemove)
+		{
+			WorldObjects->Remove(EntryToRemove);
+		}
+	}
+
+	for (TPair<uint32, FWorldObject>& WorldObjectPair : *WorldObjects.Get())
+	{
+		FWorldObject& WorldObject = WorldObjectPair.Value;
+		FTransform SimSpaceTM = ConvertWorldTransformToSimSpace(SolverSettings, WorldObject.ComponentWorldTransform);
+		if (WorldObject.LastSeenUpdateCounter == -1)
+		{
+			UE_LOG(LogRigPhysics, Log,
+				TEXT("Control rig %s - Adding object %s"), *OwnerName.ToString(),
+				*WorldObject.ActorHandle->GetName().ToString());
+
+			WorldObject.ActorHandle->SetWorldTransform(SimSpaceTM);
+			// If the object is found in the next overlap, then its update counter will be updated.
+			// However, in case it's not, we set it here too to make sure that it expires. 
+			WorldObject.LastSeenUpdateCounter = UpdateCounter;
+		}
+		else
+		{
+			WorldObject.ActorHandle->SetKinematicTarget(SimSpaceTM);
+		}
+	}
+
+}
+
+//======================================================================================================================
+// Note that the simulation mutex will already be taken when we are called
+void FRigPhysicsSimulation::UpdateWorldObjectsPostPhysics(
+	const UWorld* World, const FRigPhysicsSolverSettings& SolverSettings, const AActor* OwningActorPtr)
+{
+	WorldOverlapBox = FBox();
+	if (SolverSettings.WorldCollisionType == ERigPhysicsWorldCollisionType::None)
+	{
+		return;
+	}
+
+	// Get the volume over which to do overlaps
+	const FTransform SpaceTransform = 
+		ConvertSimSpaceTransformToComponentSpace(SolverSettings, FTransform()) * SimulationSpaceState.ComponentTM;
+	for (TPair<FRigComponentKey, FRigBodyRecord>& BodyRecordPair : BodyRecords)
+	{
+		const FRigComponentKey& ComponentKey = BodyRecordPair.Key;
+		const FRigBodyRecord& Record = BodyRecordPair.Value;
+
+		if (const ImmediatePhysics::FActorHandle* ActorHandle = Record.ActorHandle)
+		{
+			if (!ActorHandle->GetIsKinematic())
+			{
+				const Chaos::FGeometryParticleHandle* GeometryParticleHandle = ActorHandle->GetParticle();
+				const FTransform ParticleTransform = GeometryParticleHandle->GetTransformXR() * SpaceTransform;
+
+				const auto& AABB = GeometryParticleHandle->LocalBounds();
+				for (int32 Index = 0 ; Index != 8 ; ++Index)
+				{
+					WorldOverlapBox += ParticleTransform.TransformPosition(AABB.GetVertex(Index));
+				}
+			}
+		}
+	}
+
+	if (!WorldOverlapBox.IsValid)
+	{
+		return;
+	}
+
+	TWeakPtr<FCriticalSection> WeakSimulationMutex = SimulationMutex;
+
+	TWeakPtr<TMap<uint32, FWorldObject>> WeakWorldObjects = WorldObjects;
+	int64 UpdateCounterCopy = UpdateCounter;
+	const AActor* OwningActorPtrCopy = OwningActorPtr;
+	TSharedPtr<ImmediatePhysics::FSimulation> SimulationCopy = Simulation;
+
+	const TWeakObjectPtr<const UWorld> WeakWorld = World;
+
+	FVector NewHalfExtent = WorldOverlapBox.GetExtent() * SolverSettings.WorldCollisionBoundsExpansion;
+	FVector OverlapAABoxCenter = WorldOverlapBox.GetCenter();
+	WorldOverlapBox = FBox(OverlapAABoxCenter - NewHalfExtent, OverlapAABoxCenter + NewHalfExtent);
+	FCollisionShape OverlapAABox = FCollisionShape::MakeBox(NewHalfExtent);
+
+	FCollisionQueryParams QueryParams = FCollisionQueryParams(
+		SCENE_QUERY_STAT(ControlRigPhysicsFindGeometry), /*bTraceComplex=*/false);
+	QueryParams.MobilityType = EQueryMobilityType::Any;
+
+	FCollisionObjectQueryParams ObjectQueryParams;
+	switch (SolverSettings.WorldCollisionType)
+	{
+	case ERigPhysicsWorldCollisionType::Static:
+		ObjectQueryParams = FCollisionObjectQueryParams(FCollisionObjectQueryParams::InitType::AllStaticObjects); break;
+	case ERigPhysicsWorldCollisionType::Dynamic:
+		ObjectQueryParams = FCollisionObjectQueryParams(FCollisionObjectQueryParams::InitType::AllDynamicObjects); break;
+	case ERigPhysicsWorldCollisionType::All:
+		ObjectQueryParams = FCollisionObjectQueryParams(FCollisionObjectQueryParams::InitType::AllObjects); break;
+	default:
+		break;
+	}
+
+	// The actual overlap needs to be on the game thread, so we will run it in an async task.
+	// However, note that between now and then, we/our simulation etc might get destroyed so we must
+	// protect against that. Even the mutex guarding against access to the simulation may be released.
+	if (ObjectQueryParams.ObjectTypesToQuery)
+	{
+		AsyncTask(ENamedThreads::GameThread,
+			[WeakSimulationMutex, WeakWorld, WeakWorldObjects, OverlapAABox, OverlapAABoxCenter, QueryParams, 
+			ObjectQueryParams, UpdateCounterCopy, OwningActorPtrCopy, SimulationCopy]()
+			{
+				if (TSharedPtr<FCriticalSection> SimulationMutex = WeakSimulationMutex.Pin())
+				{
+					FScopeLock Lock(SimulationMutex.Get());
+
+					if (TStrongObjectPtr<const UWorld> World = WeakWorld.Pin())
+					{
+						if (TSharedPtr<TMap<uint32, FWorldObject>> WorldObjects = WeakWorldObjects.Pin())
+						{
+							TArray<FOverlapResult> Overlaps;
+							World->OverlapMultiByObjectType(
+								Overlaps, OverlapAABoxCenter, FQuat::Identity, ObjectQueryParams,
+								OverlapAABox, QueryParams);
+
+							for (const FOverlapResult& Overlap : Overlaps)
+							{
+								if (UPrimitiveComponent* OverlapComp = Overlap.GetComponent())
+								{
+									uint32 Key = OverlapComp->GetUniqueID();
+									FWorldObject* WorldObject = WorldObjects->Find(Key);
+									if (!WorldObject)
+									{
+										// New object - add it to the sim
+										if (OverlapComp->GetOwner() && OverlapComp->GetOwner() != OwningActorPtrCopy)
+										{
+											WorldObject = &WorldObjects->Add(Key, FWorldObject());
+
+											WorldObject->WorldPrimitiveComponent = OverlapComp;
+
+											// Note that we need to create the simulation actor here, even
+											// though we're in the game thread, so we can copy data out of the
+											// object. This should be safe because we only run this task after
+											// we've finished using the simulation in the worker thread.
+
+											WorldObject->ActorHandle = SimulationCopy->CreateActor(
+												ImmediatePhysics::MakeKinematicActorSetup(
+													&OverlapComp->BodyInstance, OverlapComp->GetComponentTransform()));
+
+#if WITH_EDITOR
+											WorldObject->ActorHandle->SetName(*OverlapComp->GetOwner()->GetActorLabel());
+#endif
+											SimulationCopy->AddToCollidingPairs(WorldObject->ActorHandle);
+
+											WorldObject->ComponentWorldTransform = 
+												OverlapComp->BodyInstance.GetUnrealWorldTransform();
+
+											// Flag that the simulation TM needs to be set, rather
+											// than using a kinematic target.
+											WorldObject->LastSeenUpdateCounter = -1;
+										}
+									}
+									else
+									{
+										if (TStrongObjectPtr<UPrimitiveComponent> WPC = 
+											WorldObject->WorldPrimitiveComponent.Pin())
+										{
+											WorldObject->ComponentWorldTransform = 
+												WPC->BodyInstance.GetUnrealWorldTransform();
+											WorldObject->LastSeenUpdateCounter = UpdateCounterCopy;
+										}
+									}
+								}
+							}
+						}
+					}
+				}
+			});
+	}
+}
+
diff --git a/Engine/Plugins/Experimental/ControlRigPhysics/Source/ControlRigPhysics/Public/RigPhysicsData.h b/Engine/Plugins/Experimental/ControlRigPhysics/Source/ControlRigPhysics/Public/RigPhysicsData.h
--- a/Engine/Plugins/Experimental/ControlRigPhysics/Source/ControlRigPhysics/Public/RigPhysicsData.h
+++ b/Engine/Plugins/Experimental/ControlRigPhysics/Source/ControlRigPhysics/Public/RigPhysicsData.h
@@ -51,6 +51,12 @@ struct FRigPhysicsVisualizationSettings
 	UPROPERTY(EditAnywhere, Category = Visualization)
 	bool bShowControls = false;
 
+	UPROPERTY(EditAnywhere, Category = Visualization)
+	bool bShowWorldObjects = true;
+
+	UPROPERTY(EditAnywhere, Category = Visualization)
+	bool bShowWorldOverlapBox = false;
+
 	// If visualization is enabled, show active contacts. Note that this can be overridden using
 	// CVar ControlRig.Physics.DrawActiveContactsOveride
 	UPROPERTY(EditAnywhere, Category = Visualization)
@@ -62,6 +68,23 @@ struct FRigPhysicsVisualizationSettings
 	bool bShowInactiveContacts = false;
 };
 
+/**
+ * Specifies whether/how simulation bodies should collide with the world
+ */
+UENUM()
+enum class ERigPhysicsWorldCollisionType : uint8
+{
+	// World collision is disabled
+	None,
+	// Control rig simulation bodies will only collide against static (or kinematic, since they may
+	// move) objects in the world
+	Static,
+	// Control rig simulation bodies will only collide against dynamic/simulated objects in the world
+	Dynamic,
+	// Control rig simulation bodies will collide against all objects in the world
+	All,
+};
+
 /** 
  * Specifies what space a kinematic target is defined in 
  * Danny TODO this is much better than the similar mechanism in PCC and RBWC (which has a 
@@ -442,6 +465,18 @@ struct FRigPhysicsSolverSettings
 	UPROPERTY(EditAnywhere, Category = SolverSettings)
 	FVector Gravity = FVector(0, 0, -981.0);
 
+	// Overlap type used to collect world objects. If set to None then there will be no world collisions.
+	UPROPERTY(EditAnywhere, Category = SolverSettings)
+	ERigPhysicsWorldCollisionType WorldCollisionType = ERigPhysicsWorldCollisionType::None;
+
+	// Keeps world objects for a few frames even if they go out of bounds, to avoid rapid creation/destruction.
+	UPROPERTY(EditAnywhere, Category = SolverSettings, meta = (ClampMin = 1))
+	int32 WorldCollisionExpiryFrames = 3;
+
+	// Scales up the volume used to collect world objects for collision
+	UPROPERTY(EditAnywhere, Category = SolverSettings, meta = (ClampMin = 1))
+	float WorldCollisionBoundsExpansion = 1.1f;
+
 	/**
 	 * The number of position iterations to run. The position solve is responsible for de-penetration.
 	 * Increasing this will improve simulation stability, but increase the cost.
@@ -471,7 +506,7 @@ struct FRigPhysicsSolverSettings
 	UPROPERTY(EditAnywhere, Category = SolverSettings, meta = (ClampMin = 1))
 	int32 MaxNumRollingAverageStepTimes = 1;
 
-	// This is the margin around shapes used to detect collisions. Increasing this will increase the
+	// This is the distance margin around shapes used to detect collisions. Increasing this will increase the
 	// number of inactive contacts, which will reduce the likelihood of penetration, but will also
 	// increase solver cost.
 	UPROPERTY(EditAnywhere, Category = SolverSettings, meta = (ClampMin = 0))
diff --git a/Engine/Plugins/Experimental/PhysicsControl/Source/PhysicsControl/Public/PhysicsControlObjectVersion.h b/Engine/Plugins/Experimental/PhysicsControl/Source/PhysicsControl/Public/PhysicsControlObjectVersion.h
--- a/Engine/Plugins/Experimental/PhysicsControl/Source/PhysicsControl/Public/PhysicsControlObjectVersion.h
+++ b/Engine/Plugins/Experimental/PhysicsControl/Source/PhysicsControl/Public/PhysicsControlObjectVersion.h
@@ -11,52 +11,30 @@ struct FPhysicsControlObjectVersion
 	{
 		// Before any version changes were made
 		BeforeCustomVersionWasAdded,
-
 		ControlRigIncludeDriveInJoint,
-
 		ControlRigRemoveCurrentDataFromPhysicsComponent,
-
 		ControlRigUseAutomaticSolver,
-
 		ControlRigAutomaticallyAddPhysicsComponents,
-
 		ControlRigControlAddChildBodyComponentKey,
-
 		ControlRigControlAddControlMultiplier,
-
 		ControlRigSeparateOutJointFromBody,
-
 		ControlRigSupportFullConstraintData,
-
 		ControlRigSupportFullDriveConstraintData,
-
 		ControlRigSupportNoCollisionBodies,
-
 		ControlRigSupportBodyDamping,
-
 		ControlRigCollisionHasMaterial,
-
 		ControlRigBodyDynamicsHasDensity,
-
 		ControlRigSolverSettingsIncludesCollisionBoundsExpansion,
-
 		ControlRigDetectTeleportFromDistanceChange,
-
 		ControlRigSpeedThresholdForReset,
-
 		ControlRigResetCooldownFrames,
-
 		ControlRigSimSpaceDragMultipliers,
-
 		ControlRigAccelerationThresholdForReset,
-
 		ControlRigRemoveResetCooldownFrames,
-
 		ControlRigDriveRelativeToAnimation,
-
 		ControlRigDetectExternalVelocityTurbulence,
-
 		ControlRigBodyIncludeInChecksForReset,
+		ControlRigWorldCollision,
 
 		// -----<new versions can be added above this line>-------------------------------------------------
 		VersionPlusOne,
diff --git a/Engine/Plugins/Runtime/RigVM/Source/RigVM/Public/RigVMCore/RigVMDrawInterface.h b/Engine/Plugins/Runtime/RigVM/Source/RigVM/Public/RigVMCore/RigVMDrawInterface.h
--- a/Engine/Plugins/Runtime/RigVM/Source/RigVM/Public/RigVMCore/RigVMDrawInterface.h
+++ b/Engine/Plugins/Runtime/RigVM/Source/RigVM/Public/RigVMCore/RigVMDrawInterface.h
@@ -21,6 +21,7 @@ public:
 	UE_API void DrawLine(const FTransform& WorldOffset, const FVector& LineStart, const FVector& LineEnd, const FLinearColor& Color, float Thickness = 0.f, ESceneDepthPriorityGroup DepthPriority = ESceneDepthPriorityGroup::SDPG_Foreground, float Lifetime = -1.f);
 	UE_API void DrawLines(const FTransform& WorldOffset, const TArrayView<const FVector>& Positions, const FLinearColor& Color, float Thickness = 0.f, ESceneDepthPriorityGroup DepthPriority = ESceneDepthPriorityGroup::SDPG_Foreground, float Lifetime = -1.f);
 	UE_API void DrawLineStrip(const FTransform& WorldOffset, const TArrayView<const FVector>& Positions, const FLinearColor& Color, float Thickness = 0.f, ESceneDepthPriorityGroup DepthPriority = ESceneDepthPriorityGroup::SDPG_Foreground, float Lifetime = -1.f);
+	// This draws a scaled unit cube, with the transforms applied. When drawing a box, put the position/orientation into WorldOffset and the extents into Transform.
 	UE_API void DrawBox(const FTransform& WorldOffset, const FTransform& Transform, const FLinearColor& Color, float Thickness = 0.f, ESceneDepthPriorityGroup DepthPriority = ESceneDepthPriorityGroup::SDPG_Foreground, float Lifetime = -1.f);
 	UE_API void DrawSphere(const FTransform& WorldOffset, const FTransform& Transform, float Radius, const FLinearColor& Color, float Thickness = 0.f, int32 Detail = 12, ESceneDepthPriorityGroup DepthPriority = ESceneDepthPriorityGroup::SDPG_Foreground, float Lifetime = -1.f);
 	// Hemispheres are drawn so that they extend up along +ve Z
