Control Rig Physics: Fix for the wrong solver component key being set in bodies

diff --git a/Engine/Plugins/Experimental/ControlRigPhysics/Source/ControlRigPhysics/Private/RigPhysicsExecution.cpp b/Engine/Plugins/Experimental/ControlRigPhysics/Source/ControlRigPhysics/Private/RigPhysicsExecution.cpp
--- a/Engine/Plugins/Experimental/ControlRigPhysics/Source/ControlRigPhysics/Private/RigPhysicsExecution.cpp
+++ b/Engine/Plugins/Experimental/ControlRigPhysics/Source/ControlRigPhysics/Private/RigPhysicsExecution.cpp
@@ -426,7 +426,7 @@ FRigUnit_HierarchyInstantiateFromPhysicsAsset_Execute()
 		if (FRigPhysicsBodyComponent* Component = Cast<FRigPhysicsBodyComponent>(
 			ExecuteContext.Hierarchy->FindComponent(PhysicsBodyComponentKey)))
 		{
-			Component->BodySolverSettings.PhysicsSolverComponentKey = PhysicsBodyComponentKey;
+			Component->BodySolverSettings.PhysicsSolverComponentKey = PhysicsSolverComponentKey;
 			Component->BodySolverSettings.bUseAutomaticSolver = bUseAutomaticSolver;
 			Component->Dynamics = Dynamics;
 			Component->BodyData = BodyData;
@@ -675,7 +675,7 @@ FRigUnit_HierarchyImportCollisionFromPhysicsAsset_Execute()
 			if (FRigPhysicsBodyComponent* Component = Cast<FRigPhysicsBodyComponent>(
 				ExecuteContext.Hierarchy->FindComponent(PhysicsBodyComponentKey)))
 			{
-				Component->BodySolverSettings.PhysicsSolverComponentKey = PhysicsBodyComponentKey;
+				Component->BodySolverSettings.PhysicsSolverComponentKey = PhysicsSolverComponentKey;
 				Component->BodySolverSettings.bUseAutomaticSolver = bUseAutomaticSolver;
 				Component->Dynamics = Dynamics;
 				Component->BodyData = BodyData;
