Control Rig Physics: Supports the centre of mass "nudge" - primarily needed so we can instantiate from existing physics assets and maintain behaviour.

diff --git a/Engine/Plugins/Experimental/ControlRigPhysics/Source/ControlRigPhysics/Private/RigPhysicsData.cpp b/Engine/Plugins/Experimental/ControlRigPhysics/Source/ControlRigPhysics/Private/RigPhysicsData.cpp
--- a/Engine/Plugins/Experimental/ControlRigPhysics/Source/ControlRigPhysics/Private/RigPhysicsData.cpp
+++ b/Engine/Plugins/Experimental/ControlRigPhysics/Source/ControlRigPhysics/Private/RigPhysicsData.cpp
@@ -443,6 +443,10 @@ FArchive& operator <<(FArchive& Ar, FRigPhysicsDynamics& Data)
 	Ar << Data.MassOverride;
 	Ar << Data.bOverrideCentreOfMass;
 	Ar << Data.CentreOfMassOverride;
+	if (Ar.CustomVer(FPhysicsControlObjectVersion::GUID) >= FPhysicsControlObjectVersion::ControlRigCentreOfMassNudge)
+	{
+		Ar << Data.CentreOfMassNudge;
+	}
 	Ar << Data.bOverrideMomentsOfInertia;
 	Ar << Data.MomentsOfInertiaOverride;
 	if (Ar.CustomVer(FPhysicsControlObjectVersion::GUID) >= FPhysicsControlObjectVersion::ControlRigSupportBodyDamping)
diff --git a/Engine/Plugins/Experimental/ControlRigPhysics/Source/ControlRigPhysics/Private/RigPhysicsExecution.cpp b/Engine/Plugins/Experimental/ControlRigPhysics/Source/ControlRigPhysics/Private/RigPhysicsExecution.cpp
--- a/Engine/Plugins/Experimental/ControlRigPhysics/Source/ControlRigPhysics/Private/RigPhysicsExecution.cpp
+++ b/Engine/Plugins/Experimental/ControlRigPhysics/Source/ControlRigPhysics/Private/RigPhysicsExecution.cpp
@@ -410,6 +410,7 @@ FRigUnit_HierarchyInstantiateFromPhysicsAsset_Execute()
 		Dynamics.MassOverride = SkeletalBodySetup->CalculateMass();
 		Dynamics.LinearDamping = SkeletalBodySetup->DefaultInstance.LinearDamping;
 		Dynamics.AngularDamping = SkeletalBodySetup->DefaultInstance.AngularDamping;
+		Dynamics.CentreOfMassNudge = SkeletalBodySetup->DefaultInstance.COMNudge;
 
 		// Body Data
 		FPhysicsControlModifierData BodyData;
diff --git a/Engine/Plugins/Experimental/ControlRigPhysics/Source/ControlRigPhysics/Private/RigPhysicsSimulation.cpp b/Engine/Plugins/Experimental/ControlRigPhysics/Source/ControlRigPhysics/Private/RigPhysicsSimulation.cpp
--- a/Engine/Plugins/Experimental/ControlRigPhysics/Source/ControlRigPhysics/Private/RigPhysicsSimulation.cpp
+++ b/Engine/Plugins/Experimental/ControlRigPhysics/Source/ControlRigPhysics/Private/RigPhysicsSimulation.cpp
@@ -562,17 +562,30 @@ ImmediatePhysics::FActorHandle* FRigPhysicsSimulation::CreateBody(
 	// We will have created with an arbitrary density - adjust to result in the desired mass.
 	ActorSetup.Mass = Mass;
 	ActorSetup.Inertia = Inertia;
-	if (Mass > 0.0f && Dynamics && Dynamics->MassOverride > 0.0f)
-	{
-		ActorSetup.Mass = Dynamics->MassOverride;
-		ActorSetup.Inertia = (Dynamics->MassOverride / Mass) * Inertia;
-	}
-
 	ActorSetup.Transform = BodyRelSimSpaceTM;
 	ActorSetup.CoMTransform = CoMTransform;
 	ActorSetup.Geometry = MoveTemp(BodyGeom);
 	ActorSetup.Shapes = MoveTemp(BodyShapes);
 
+	if (Dynamics)
+	{
+		if (Mass > 0.0f && Dynamics->MassOverride > 0.0f)
+		{
+			ActorSetup.Mass = Dynamics->MassOverride;
+			ActorSetup.Inertia = (Dynamics->MassOverride / Mass) * Inertia;
+		}
+		if (Dynamics->bOverrideMomentsOfInertia)
+		{
+			ActorSetup.Inertia = Dynamics->MomentsOfInertiaOverride;
+		}
+		if (Dynamics->bOverrideCentreOfMass)
+		{
+			ActorSetup.CoMTransform.SetLocation(Dynamics->CentreOfMassOverride);
+		}
+		ActorSetup.CoMTransform.AddToTranslation(Dynamics->CentreOfMassNudge);
+	}
+
+
 	ActorSetup.Material = MakeUnique<Chaos::FChaosPhysicsMaterial>();
 	ActorSetup.Material->Friction = Collision.Material.Friction;
 	ActorSetup.Material->StaticFriction = Collision.Material.Friction;
diff --git a/Engine/Plugins/Experimental/ControlRigPhysics/Source/ControlRigPhysics/Private/RigPhysicsSimulationVisualization.cpp b/Engine/Plugins/Experimental/ControlRigPhysics/Source/ControlRigPhysics/Private/RigPhysicsSimulationVisualization.cpp
--- a/Engine/Plugins/Experimental/ControlRigPhysics/Source/ControlRigPhysics/Private/RigPhysicsSimulationVisualization.cpp
+++ b/Engine/Plugins/Experimental/ControlRigPhysics/Source/ControlRigPhysics/Private/RigPhysicsSimulationVisualization.cpp
@@ -153,6 +153,14 @@ static void DrawActor(
 			DI, VisualizationSettings, bIsKinematic, bCollides, GeometryParticleHandle,
 			ParticleTransform, ImplicitObject, ShapeInstance.Get());
 	}
+
+	if (VisualizationSettings.bShowCentreOfMass)
+	{
+		float Size = 1.0f * VisualizationSettings.ShapeSize;
+		DI.DrawAxes(
+			ActorHandle.GetWorldTransform(), ActorHandle.GetLocalCoMTransform(), 
+			Size, VisualizationSettings.LineThickness);
+	}
 }
 
 enum
diff --git a/Engine/Plugins/Experimental/ControlRigPhysics/Source/ControlRigPhysics/Public/RigPhysicsData.h b/Engine/Plugins/Experimental/ControlRigPhysics/Source/ControlRigPhysics/Public/RigPhysicsData.h
--- a/Engine/Plugins/Experimental/ControlRigPhysics/Source/ControlRigPhysics/Public/RigPhysicsData.h
+++ b/Engine/Plugins/Experimental/ControlRigPhysics/Source/ControlRigPhysics/Public/RigPhysicsData.h
@@ -45,6 +45,9 @@ struct FRigPhysicsVisualizationSettings
 	UPROPERTY(EditAnywhere, Category = Visualization)
 	bool bShowBodies = true;
 
+	UPROPERTY(EditAnywhere, Category = Visualization)
+	bool bShowCentreOfMass = false;
+
 	UPROPERTY(EditAnywhere, Category = Visualization)
 	bool bShowJoints = true;
 
@@ -613,6 +616,10 @@ struct FRigPhysicsDynamics
 	UPROPERTY(EditAnywhere, Category = Dynamics, meta = (EditCondition = "bOverrideCentreOfMass"))
 	FVector CentreOfMassOverride = FVector::ZeroVector;
 
+	// Offset the centre of mass from the calculated (or overridden) position.
+	UPROPERTY(EditAnywhere, Category = Dynamics)
+	FVector CentreOfMassNudge = FVector::ZeroVector;
+
 	// Use this to specify the moments of inertia
 	UPROPERTY(EditAnywhere, Category = Dynamics)
 	bool bOverrideMomentsOfInertia = false;
diff --git a/Engine/Plugins/Experimental/PhysicsControl/Source/PhysicsControl/Public/PhysicsControlObjectVersion.h b/Engine/Plugins/Experimental/PhysicsControl/Source/PhysicsControl/Public/PhysicsControlObjectVersion.h
--- a/Engine/Plugins/Experimental/PhysicsControl/Source/PhysicsControl/Public/PhysicsControlObjectVersion.h
+++ b/Engine/Plugins/Experimental/PhysicsControl/Source/PhysicsControl/Public/PhysicsControlObjectVersion.h
@@ -35,6 +35,7 @@ struct FPhysicsControlObjectVersion
 		ControlRigDetectExternalVelocityTurbulence,
 		ControlRigBodyIncludeInChecksForReset,
 		ControlRigWorldCollision,
+		ControlRigCentreOfMassNudge,
 
 		// -----<new versions can be added above this line>-------------------------------------------------
 		VersionPlusOne,
