Control Rig Physics: Trivial fix to ensure the simulation space is incorporated in debug draw

diff --git a/Engine/Plugins/Experimental/ControlRigPhysics/Source/ControlRigPhysics/Private/RigPhysicsSimulationVisualization.cpp b/Engine/Plugins/Experimental/ControlRigPhysics/Source/ControlRigPhysics/Private/RigPhysicsSimulationVisualization.cpp
--- a/Engine/Plugins/Experimental/ControlRigPhysics/Source/ControlRigPhysics/Private/RigPhysicsSimulationVisualization.cpp
+++ b/Engine/Plugins/Experimental/ControlRigPhysics/Source/ControlRigPhysics/Private/RigPhysicsSimulationVisualization.cpp
@@ -158,11 +158,12 @@ static void DrawActor(
 	{
 		float Size = 1.0f * VisualizationSettings.ShapeSize;
 		DI.DrawAxes(
-			ActorHandle.GetWorldTransform(), ActorHandle.GetLocalCoMTransform(), 
+			ParticleTransform, ActorHandle.GetLocalCoMTransform(),
 			Size, VisualizationSettings.LineThickness);
 	}
 }
 
+//======================================================================================================================
 enum
 {
 	TwistIndex = 0,
@@ -170,7 +171,8 @@ enum
 	Swing2Index = 2
 };
 
-// Returns a number to use for the limit for limit based on limit type
+//======================================================================================================================
+// Returns the limit range - it will be the value passed in unless it the motion is free or locked. 
 float GetLimitAngleRadians(float LimitAngle, Chaos::EJointMotionType LimitType)
 {
 	switch (LimitType)
