Control Rig Physics - Change the update/evaluation of controls to avoid unwanted side effects that can happen when the objects being controlled are kinematic. This prevents controls "popping" when switching bodies between dynamic and kinematic.

diff --git a/Engine/Plugins/Experimental/ControlRigPhysics/Source/ControlRigPhysics/Private/RigPhysicsSimulation.cpp b/Engine/Plugins/Experimental/ControlRigPhysics/Source/ControlRigPhysics/Private/RigPhysicsSimulation.cpp
--- a/Engine/Plugins/Experimental/ControlRigPhysics/Source/ControlRigPhysics/Private/RigPhysicsSimulation.cpp
+++ b/Engine/Plugins/Experimental/ControlRigPhysics/Source/ControlRigPhysics/Private/RigPhysicsSimulation.cpp
@@ -1400,85 +1400,95 @@ void FRigPhysicsSimulation::UpdateControlPrePhysics(
 	Constraint->SetCollisionEnabled(!ControlComponent->ControlData.bDisableCollision);
 	Constraint->SetParentInvMassScale(ControlComponent->ControlData.bOnlyControlChildObject ? 0 : 1);
 
+	// Check for early outs. Make sure that we don't apply velocities using the wrong calculation
+	// when the strength/damping is increased in the future by not updating the update counter.
+
+	const ImmediatePhysics::FActorHandle* ChildActorHandle =
+		JointHandle->GetActorHandles()[ConstraintChildIndex];
+	const ImmediatePhysics::FActorHandle* ParentActorHandle =
+		JointHandle->GetActorHandles()[ConstraintParentIndex];
+
+	if (!ChildActorHandle || !ParentActorHandle)
+	{
+		return;
+	}
+
 	const Chaos::FPBDJointSettings& JointSettings = Constraint->GetSettings();
-	if (UpdateDriveSpringDamperSettings(
+
+	// Note that if we don't have any strength, then we don't calculate the targets. 
+	if (!UpdateDriveSpringDamperSettings(
 		JointHandle, JointSettings, ControlComponent->ControlData, ControlComponent->ControlMultiplier))
 	{
-		const ImmediatePhysics::FActorHandle* ChildActorHandle =
-			JointHandle->GetActorHandles()[ConstraintChildIndex];
-		const ImmediatePhysics::FActorHandle* ParentActorHandle =
-			JointHandle->GetActorHandles()[ConstraintParentIndex];
+		return;
+	}
 
-		if (ChildActorHandle && ParentActorHandle)
-		{
-			// TODO
-			// - cache settings / previous input parameters to avoid unnecessary repeating
-			//   calculations and making physics API calls every update.
+	// TODO
+	// - cache settings / previous input parameters to avoid unnecessary repeating
+	//   calculations and making physics API calls every update.
 
-			// Update the target point on the child
-			Constraint->SetChildConnectorLocation(
-				ControlComponent->ControlData.GetControlPoint(ChildActorHandle));
+	// Update the target point on the child
 
-			FTransform TargetTM(
-				ControlComponent->ControlTarget.TargetOrientation,
-				ControlComponent->ControlTarget.TargetPosition);
+	// Danny TODO Note that if child is kinematic then getting the control point will be a problem
+	// as Chaos doesn't like to return the CoM for kinematics. Consider changing Chaos so that the
+	// CoM can still be queried. This workaround (of skipping the call) will result in bad behaviour
+	// if the child is kinematic and the parent is dynamic. 
+	if (!ChildActorHandle->GetIsKinematic())
+	{
+		Constraint->SetChildConnectorLocation(
+			ControlComponent->ControlData.GetControlPoint(ChildActorHandle));
+	}
 
-			if (ControlComponent->ControlData.bUseSkeletalAnimation)
-			{
-				const FTransform ComponentSpaceAnimTargetTM = CalculateTargetTM(
-					Hierarchy, JointSettings, ControlRecord).ToTransform();
-				const FTransform SimSpaceAnimTargetTM = ConvertComponentSpaceTransformToSimSpace(
-					SolverSettings, ComponentSpaceAnimTargetTM);
-				TargetTM = TargetTM * SimSpaceAnimTargetTM;
-			}
-			else
-			{
-				TargetTM = ConvertComponentSpaceTransformToSimSpace(SolverSettings, TargetTM);
-			}
+	FTransform TargetTM(
+		ControlComponent->ControlTarget.TargetOrientation,
+		ControlComponent->ControlTarget.TargetPosition);
 
-			Constraint->SetLinearDrivePositionTarget(TargetTM.GetTranslation());
-			Constraint->SetAngularDrivePositionTarget(TargetTM.GetRotation());
+	if (ControlComponent->ControlData.bUseSkeletalAnimation)
+	{
+		const FTransform ComponentSpaceAnimTargetTM = CalculateTargetTM(
+			Hierarchy, JointSettings, ControlRecord).ToTransform();
+		const FTransform SimSpaceAnimTargetTM = ConvertComponentSpaceTransformToSimSpace(
+			SolverSettings, ComponentSpaceAnimTargetTM);
+		TargetTM = TargetTM * SimSpaceAnimTargetTM;
+	}
+	else
+	{
+		TargetTM = ConvertComponentSpaceTransformToSimSpace(SolverSettings, TargetTM);
+	}
 
-			if ((ThisDeltaTime * ControlComponent->ControlData.LinearTargetVelocityMultiplier) != 0)
-			{
-				FVector Velocity =
-					(TargetTM.GetTranslation() - ControlRecord.PreviousTargetTM.GetTranslation()) / ThisDeltaTime;
-				Constraint->SetLinearDriveVelocityTarget(
-					Velocity * ControlComponent->ControlData.LinearTargetVelocityMultiplier);
-			}
-			else
-			{
-				Constraint->SetLinearDriveVelocityTarget(Chaos::FVec3(0));
-			}
+	Constraint->SetLinearDrivePositionTarget(TargetTM.GetTranslation());
+	Constraint->SetAngularDrivePositionTarget(TargetTM.GetRotation());
 
-			if ((ThisDeltaTime * ControlComponent->ControlData.AngularTargetVelocityMultiplier) != 0)
-			{
-				// Note that quats multiply in the opposite order to TMs, and must be in the same hemisphere.
-				const FQuat Q = TargetTM.GetRotation();
-				FQuat PrevQ = ControlRecord.PreviousTargetTM.GetRotation();
-				PrevQ.EnforceShortestArcWith(Q);
-				const FQuat DeltaQ = Q * PrevQ.Inverse();
-				const FVector AngularVelocity = DeltaQ.ToRotationVector() / ThisDeltaTime;
-
-				Constraint->SetAngularDriveVelocityTarget(
-					AngularVelocity * ControlComponent->ControlData.AngularTargetVelocityMultiplier);
-			}
-			else
-			{
-				Constraint->SetAngularDriveVelocityTarget(Chaos::FVec3(0));
-			}
+	if ((ThisDeltaTime * ControlComponent->ControlData.LinearTargetVelocityMultiplier) != 0)
+	{
+		FVector Velocity =
+			(TargetTM.GetTranslation() - ControlRecord.PreviousTargetTM.GetTranslation()) / ThisDeltaTime;
+		Constraint->SetLinearDriveVelocityTarget(
+			Velocity * ControlComponent->ControlData.LinearTargetVelocityMultiplier);
+	}
+	else
+	{
+		Constraint->SetLinearDriveVelocityTarget(Chaos::FVec3(0));
+	}
 
+	if ((ThisDeltaTime * ControlComponent->ControlData.AngularTargetVelocityMultiplier) != 0)
+	{
+		// Note that quats multiply in the opposite order to TMs, and must be in the same hemisphere.
+		const FQuat Q = TargetTM.GetRotation();
+		FQuat PrevQ = ControlRecord.PreviousTargetTM.GetRotation();
+		PrevQ.EnforceShortestArcWith(Q);
+		const FQuat DeltaQ = Q * PrevQ.Inverse();
+		const FVector AngularVelocity = DeltaQ.ToRotationVector() / ThisDeltaTime;
 
-			ControlRecord.PreviousTargetTM = TargetTM;
-			ControlRecord.PreviousTargetUpdateCounter = UpdateCounter;
-		}
-		else
-		{
-			// Note that if we don't have any strength, then we don't calculate the targets.
-			// However, make sure that we don't apply velocities using the wrong calculation
-			// when the strength/damping is increased in the future
-		}
+		Constraint->SetAngularDriveVelocityTarget(
+			AngularVelocity * ControlComponent->ControlData.AngularTargetVelocityMultiplier);
 	}
+	else
+	{
+		Constraint->SetAngularDriveVelocityTarget(Chaos::FVec3(0));
+	}
+
+	ControlRecord.PreviousTargetTM = TargetTM;
+	ControlRecord.PreviousTargetUpdateCounter = UpdateCounter;
 }
 
 //======================================================================================================================
