Control Rig Physics: Expose some functions to adjust control settings

diff --git a/Engine/Plugins/Experimental/ControlRigPhysics/Source/ControlRigPhysics/Private/RigPhysicsControlExecution.cpp b/Engine/Plugins/Experimental/ControlRigPhysics/Source/ControlRigPhysics/Private/RigPhysicsControlExecution.cpp
--- a/Engine/Plugins/Experimental/ControlRigPhysics/Source/ControlRigPhysics/Private/RigPhysicsControlExecution.cpp
+++ b/Engine/Plugins/Experimental/ControlRigPhysics/Source/ControlRigPhysics/Private/RigPhysicsControlExecution.cpp
@@ -91,6 +91,76 @@ FRigUnit_HierarchySetControlData_Execute()
 	}
 }
 
+//======================================================================================================================
+FRigUnit_HierarchySetControlLinearStrength_Execute()
+{
+	DECLARE_SCOPE_HIERARCHICAL_COUNTER_RIGUNIT();
+	if (ExecuteContext.Hierarchy)
+	{
+		if (FRigPhysicsControlComponent* Component = Cast<FRigPhysicsControlComponent>(
+			ExecuteContext.Hierarchy->FindComponent(PhysicsControlComponentKey)))
+		{
+			Component->ControlData.LinearStrength = Strength;
+		}
+	}
+}
+
+//======================================================================================================================
+FRigUnit_HierarchySetControlLinearDampingRatio_Execute()
+{
+	DECLARE_SCOPE_HIERARCHICAL_COUNTER_RIGUNIT();
+	if (ExecuteContext.Hierarchy)
+	{
+		if (FRigPhysicsControlComponent* Component = Cast<FRigPhysicsControlComponent>(
+			ExecuteContext.Hierarchy->FindComponent(PhysicsControlComponentKey)))
+		{
+			Component->ControlData.LinearDampingRatio = DampingRatio;
+		}
+	}
+}
+
+//======================================================================================================================
+FRigUnit_HierarchySetControlAngularStrength_Execute()
+{
+	DECLARE_SCOPE_HIERARCHICAL_COUNTER_RIGUNIT();
+	if (ExecuteContext.Hierarchy)
+	{
+		if (FRigPhysicsControlComponent* Component = Cast<FRigPhysicsControlComponent>(
+			ExecuteContext.Hierarchy->FindComponent(PhysicsControlComponentKey)))
+		{
+			Component->ControlData.AngularStrength = Strength;
+		}
+	}
+}
+
+//======================================================================================================================
+FRigUnit_HierarchySetControlAngularDampingRatio_Execute()
+{
+	DECLARE_SCOPE_HIERARCHICAL_COUNTER_RIGUNIT();
+	if (ExecuteContext.Hierarchy)
+	{
+		if (FRigPhysicsControlComponent* Component = Cast<FRigPhysicsControlComponent>(
+			ExecuteContext.Hierarchy->FindComponent(PhysicsControlComponentKey)))
+		{
+			Component->ControlData.AngularDampingRatio = DampingRatio;
+		}
+	}
+}
+
+//======================================================================================================================
+FRigUnit_HierarchyGetControlData_Execute()
+{
+	DECLARE_SCOPE_HIERARCHICAL_COUNTER_RIGUNIT();
+	if (ExecuteContext.Hierarchy)
+	{
+		if (FRigPhysicsControlComponent* Component = Cast<FRigPhysicsControlComponent>(
+			ExecuteContext.Hierarchy->FindComponent(PhysicsControlComponentKey)))
+		{
+			ControlData = Component->ControlData;
+		}
+	}
+}
+
 //======================================================================================================================
 FRigUnit_HierarchySetControlMultiplier_Execute()
 {
diff --git a/Engine/Plugins/Experimental/ControlRigPhysics/Source/ControlRigPhysics/Public/RigPhysicsControlExecution.h b/Engine/Plugins/Experimental/ControlRigPhysics/Source/ControlRigPhysics/Public/RigPhysicsControlExecution.h
--- a/Engine/Plugins/Experimental/ControlRigPhysics/Source/ControlRigPhysics/Public/RigPhysicsControlExecution.h
+++ b/Engine/Plugins/Experimental/ControlRigPhysics/Source/ControlRigPhysics/Public/RigPhysicsControlExecution.h
@@ -139,6 +139,112 @@ struct FRigUnit_HierarchySetControlData : public FRigUnit_PhysicsBaseMutable
 	FPhysicsControlData ControlData;
 };
 
+USTRUCT(meta = (DisplayName = "Set Physics Control Linear Strength", Varying))
+struct FRigUnit_HierarchySetControlLinearStrength : public FRigUnit_PhysicsBaseMutable
+{
+	GENERATED_BODY()
+
+	FRigUnit_HierarchySetControlLinearStrength()
+	{
+		PhysicsControlComponentKey.ElementKey.Type = ERigElementType::Bone;
+		PhysicsControlComponentKey.Name = FRigPhysicsControlComponent::GetDefaultName();
+	}
+
+	RIGVM_METHOD()
+	UE_API virtual void Execute() override;
+
+	UPROPERTY(meta = (Input))
+	FRigComponentKey PhysicsControlComponentKey;
+
+	UPROPERTY(meta = (Input))
+	float Strength;
+};
+
+USTRUCT(meta = (DisplayName = "Set Physics Control Linear Damping Ratio", Varying))
+struct FRigUnit_HierarchySetControlLinearDampingRatio : public FRigUnit_PhysicsBaseMutable
+{
+	GENERATED_BODY()
+
+	FRigUnit_HierarchySetControlLinearDampingRatio()
+	{
+		PhysicsControlComponentKey.ElementKey.Type = ERigElementType::Bone;
+		PhysicsControlComponentKey.Name = FRigPhysicsControlComponent::GetDefaultName();
+	}
+
+	RIGVM_METHOD()
+	UE_API virtual void Execute() override;
+
+	UPROPERTY(meta = (Input))
+	FRigComponentKey PhysicsControlComponentKey;
+
+	UPROPERTY(meta = (Input))
+	float DampingRatio;
+};
+
+USTRUCT(meta = (DisplayName = "Set Physics Control Angular Strength", Varying))
+struct FRigUnit_HierarchySetControlAngularStrength : public FRigUnit_PhysicsBaseMutable
+{
+	GENERATED_BODY()
+
+	FRigUnit_HierarchySetControlAngularStrength()
+	{
+		PhysicsControlComponentKey.ElementKey.Type = ERigElementType::Bone;
+		PhysicsControlComponentKey.Name = FRigPhysicsControlComponent::GetDefaultName();
+	}
+
+	RIGVM_METHOD()
+	UE_API virtual void Execute() override;
+
+	UPROPERTY(meta = (Input))
+	FRigComponentKey PhysicsControlComponentKey;
+
+	UPROPERTY(meta = (Input))
+	float Strength;
+};
+
+USTRUCT(meta = (DisplayName = "Set Physics Control Angular Damping Ratio", Varying))
+struct FRigUnit_HierarchySetControlAngularDampingRatio : public FRigUnit_PhysicsBaseMutable
+{
+	GENERATED_BODY()
+
+	FRigUnit_HierarchySetControlAngularDampingRatio()
+	{
+		PhysicsControlComponentKey.ElementKey.Type = ERigElementType::Bone;
+		PhysicsControlComponentKey.Name = FRigPhysicsControlComponent::GetDefaultName();
+	}
+
+	RIGVM_METHOD()
+	UE_API virtual void Execute() override;
+
+	UPROPERTY(meta = (Input))
+	FRigComponentKey PhysicsControlComponentKey;
+
+	UPROPERTY(meta = (Input))
+	float DampingRatio;
+};
+
+// Gets the control data for a physics control
+USTRUCT(meta = (DisplayName = "Get Physics Control Data", Varying))
+struct FRigUnit_HierarchyGetControlData : public FRigUnit_PhysicsBase
+{
+	GENERATED_BODY()
+
+	FRigUnit_HierarchyGetControlData()
+	{
+		PhysicsControlComponentKey.ElementKey.Type = ERigElementType::Bone;
+		PhysicsControlComponentKey.Name = FRigPhysicsControlComponent::GetDefaultName();
+	}
+
+	RIGVM_METHOD()
+	UE_API virtual void Execute() override;
+
+	UPROPERTY(meta = (Input))
+	FRigComponentKey PhysicsControlComponentKey;
+
+	UPROPERTY(meta = (Output))
+	FPhysicsControlData ControlData;
+};
+
 // Sets the multipliers for a physics control
 USTRUCT(meta = (DisplayName = "Set Physics Control Multiplier", Varying))
 struct FRigUnit_HierarchySetControlMultiplier : public FRigUnit_PhysicsBaseMutable
diff --git a/Engine/Plugins/Experimental/ControlRigPhysics/Source/ControlRigPhysics/Public/RigPhysicsData.h b/Engine/Plugins/Experimental/ControlRigPhysics/Source/ControlRigPhysics/Public/RigPhysicsData.h
--- a/Engine/Plugins/Experimental/ControlRigPhysics/Source/ControlRigPhysics/Public/RigPhysicsData.h
+++ b/Engine/Plugins/Experimental/ControlRigPhysics/Source/ControlRigPhysics/Public/RigPhysicsData.h
@@ -260,7 +260,7 @@ struct FRigPhysicsMaterial
 
 	// Friction is simplified here - just one value used for static and dynamic friction. 
 	UPROPERTY(EditAnywhere, Category = Material, meta = (ClampMin = "0.0"))
-	float Friction = 0.5f;
+	float Friction = 1.0f;
 
 	UPROPERTY(EditAnywhere, Category = Material, meta = (ClampMin = "0.0", ClampMax = "1.0"))
 	float Restitution = 0.0f;
@@ -485,14 +485,14 @@ struct FRigPhysicsSolverSettings
 	 * Increasing this will improve simulation stability, but increase the cost.
 	*/
 	UPROPERTY(EditAnywhere, Category = SolverSettings, meta = (ClampMin = 0))
-	int32 PositionIterations = 6;
+	int32 PositionIterations = 8;
 
 	/**
 	 * The number of velocity iterations to run. The velocity solve is responsible for restitution (bounce) and friction.
-	 * This should usually be 1, but could be 0 if you don't care about friction and restitution.
+	 * This should usually be a low value (even 1), but could be 0 if you don't care about friction and restitution.
 	*/
 	UPROPERTY(EditAnywhere, Category = SolverSettings, meta = (ClampMin = 0))
-	int32 VelocityIterations = 1;
+	int32 VelocityIterations = 4;
 
 	/**
 	 * The number of projection iterations to run. The projection phase is a final pass over the constraints, applying
