Control Rig Physics: Expose enable/disable on the physics joints

diff --git a/Engine/Plugins/Experimental/ControlRigPhysics/Source/ControlRigPhysics/Private/RigPhysicsData.cpp b/Engine/Plugins/Experimental/ControlRigPhysics/Source/ControlRigPhysics/Private/RigPhysicsData.cpp
--- a/Engine/Plugins/Experimental/ControlRigPhysics/Source/ControlRigPhysics/Private/RigPhysicsData.cpp
+++ b/Engine/Plugins/Experimental/ControlRigPhysics/Source/ControlRigPhysics/Private/RigPhysicsData.cpp
@@ -94,6 +94,10 @@ FArchive& operator <<(FArchive& Ar, FAngularDriveConstraint& Data)
 //======================================================================================================================
 FArchive& operator <<(FArchive& Ar, FRigPhysicsJointData& Data)
 {
+	if (Ar.CustomVer(FPhysicsControlObjectVersion::GUID) >= FPhysicsControlObjectVersion::ControlRigJointEnabled)
+	{
+		Ar << Data.bEnabled;
+	}
 	if (Ar.CustomVer(FPhysicsControlObjectVersion::GUID) < FPhysicsControlObjectVersion::ControlRigSeparateOutJointFromBody)
 	{
 		bool bEnable;
diff --git a/Engine/Plugins/Experimental/ControlRigPhysics/Source/ControlRigPhysics/Private/RigPhysicsJointExecution.cpp b/Engine/Plugins/Experimental/ControlRigPhysics/Source/ControlRigPhysics/Private/RigPhysicsJointExecution.cpp
--- a/Engine/Plugins/Experimental/ControlRigPhysics/Source/ControlRigPhysics/Private/RigPhysicsJointExecution.cpp
+++ b/Engine/Plugins/Experimental/ControlRigPhysics/Source/ControlRigPhysics/Private/RigPhysicsJointExecution.cpp
@@ -82,6 +82,21 @@ FRigUnit_HierarchyGetJointData_Execute()
 	}
 }
 
+//======================================================================================================================
+FRigUnit_HierarchySetJointEnabled_Execute()
+{
+	DECLARE_SCOPE_HIERARCHICAL_COUNTER_RIGUNIT();
+	if (!ExecuteContext.Hierarchy)
+	{
+		return;
+	}
+	if (FRigPhysicsJointComponent* Component = Cast<FRigPhysicsJointComponent>(
+		ExecuteContext.Hierarchy->FindComponent(PhysicsJointComponentKey)))
+	{
+		Component->JointData.bEnabled = bEnabled;
+	}
+}
+
 //======================================================================================================================
 FRigUnit_HierarchySetJointDriveData_Execute()
 {
diff --git a/Engine/Plugins/Experimental/ControlRigPhysics/Source/ControlRigPhysics/Private/RigPhysicsSimulation.cpp b/Engine/Plugins/Experimental/ControlRigPhysics/Source/ControlRigPhysics/Private/RigPhysicsSimulation.cpp
--- a/Engine/Plugins/Experimental/ControlRigPhysics/Source/ControlRigPhysics/Private/RigPhysicsSimulation.cpp
+++ b/Engine/Plugins/Experimental/ControlRigPhysics/Source/ControlRigPhysics/Private/RigPhysicsSimulation.cpp
@@ -384,23 +384,6 @@ static void SetCommonProperties(const FRigPhysicsCollisionShape& Shape, FKShapeE
 #endif
 }
 
-//======================================================================================================================
-static Chaos::EJointMotionType GetMotionTypeFromLimitValue(const float Value)
-{
-	if (Value > 0.0f)
-	{
-		return Chaos::EJointMotionType::Limited;
-	}
-	else if (Value < 0.0f)
-	{
-		return Chaos::EJointMotionType::Free;
-	}
-	else
-	{
-		return Chaos::EJointMotionType::Locked;
-	}
-}
-
 //======================================================================================================================
 static bool CreateGeometry(
 	const FRigPhysicsCollision&               Collision,
@@ -1092,9 +1075,12 @@ void FRigPhysicsSimulation::UpdateJointPrePhysics(
 		const FRigPhysicsJointData& JointData = PhysicsJointComponent->JointData;
 		const FRigPhysicsDriveData& DriveData = PhysicsJointComponent->DriveData;
 
+		Constraint->SetConstraintEnabled(JointData.bEnabled);
+
 		Chaos::FPBDJointSettings Settings = Constraint->GetSettings();
 		ImmediatePhysics::UpdateJointSettingsFromLinearDriveConstraint(DriveData.LinearDriveConstraint, Settings);
 		ImmediatePhysics::UpdateJointSettingsFromAngularDriveConstraint(DriveData.AngularDriveConstraint, Settings);
+
 		Constraint->SetSettings(Settings);
 
 		// Now set the actual target
diff --git a/Engine/Plugins/Experimental/ControlRigPhysics/Source/ControlRigPhysics/Public/RigPhysicsData.h b/Engine/Plugins/Experimental/ControlRigPhysics/Source/ControlRigPhysics/Public/RigPhysicsData.h
--- a/Engine/Plugins/Experimental/ControlRigPhysics/Source/ControlRigPhysics/Public/RigPhysicsData.h
+++ b/Engine/Plugins/Experimental/ControlRigPhysics/Source/ControlRigPhysics/Public/RigPhysicsData.h
@@ -680,6 +680,9 @@ struct FRigPhysicsJointData
 
 	friend FArchive& operator<<(FArchive& Ar, FRigPhysicsJointData& Data);
 
+	UPROPERTY(EditAnywhere, Category = Joint)
+	bool bEnabled = true;
+
 	// The auto-calculated offset of the parent frame from the parent body places it at the 
 	// location of the child body relative to the parent body in the initial pose.
 	UPROPERTY(EditAnywhere, Category = Joint)
diff --git a/Engine/Plugins/Experimental/ControlRigPhysics/Source/ControlRigPhysics/Public/RigPhysicsJointExecution.h b/Engine/Plugins/Experimental/ControlRigPhysics/Source/ControlRigPhysics/Public/RigPhysicsJointExecution.h
--- a/Engine/Plugins/Experimental/ControlRigPhysics/Source/ControlRigPhysics/Public/RigPhysicsJointExecution.h
+++ b/Engine/Plugins/Experimental/ControlRigPhysics/Source/ControlRigPhysics/Public/RigPhysicsJointExecution.h
@@ -106,6 +106,27 @@ struct FRigUnit_HierarchyGetJointData : public FRigUnit_PhysicsBase
 	FRigPhysicsJointData JointData;
 };
 
+USTRUCT(meta = (DisplayName = "Set Physics Joint Enabled", Varying))
+struct FRigUnit_HierarchySetJointEnabled : public FRigUnit_PhysicsBaseMutable
+{
+	GENERATED_BODY()
+
+	FRigUnit_HierarchySetJointEnabled()
+	{
+		PhysicsJointComponentKey.ElementKey.Type = ERigElementType::Bone;
+		PhysicsJointComponentKey.Name = FRigPhysicsJointComponent::GetDefaultName();
+	}
+
+	RIGVM_METHOD()
+	UE_API virtual void Execute() override;
+
+	UPROPERTY(meta = (Input))
+	FRigComponentKey PhysicsJointComponentKey;
+
+	UPROPERTY(meta = (Input))
+	bool bEnabled;
+};
+
 // Sets the joint drive for a physics component body
 USTRUCT(meta = (DisplayName = "Set Physics Joint Drive Properties", Varying))
 struct FRigUnit_HierarchySetJointDriveData : public FRigUnit_PhysicsBaseMutable
@@ -172,7 +193,6 @@ struct FRigUnit_HierarchySetJointDriveUseSkeletalAnimation : public FRigUnit_Phy
 	bool bUseSkeletalAnimation = true;
 };
 
-
 // Helper to simplify creation of joint data for a typical articulation (i.e. locked linear motion,
 // but with an angular limit)
 USTRUCT(meta = (DisplayName = "Make Articulation Joint Data", Keywords = "Construction,Create,New,Body,Skeleton", Varying))
diff --git a/Engine/Plugins/Experimental/PhysicsControl/Source/PhysicsControl/Public/PhysicsControlObjectVersion.h b/Engine/Plugins/Experimental/PhysicsControl/Source/PhysicsControl/Public/PhysicsControlObjectVersion.h
--- a/Engine/Plugins/Experimental/PhysicsControl/Source/PhysicsControl/Public/PhysicsControlObjectVersion.h
+++ b/Engine/Plugins/Experimental/PhysicsControl/Source/PhysicsControl/Public/PhysicsControlObjectVersion.h
@@ -36,6 +36,7 @@ struct FPhysicsControlObjectVersion
 		ControlRigBodyIncludeInChecksForReset,
 		ControlRigWorldCollision,
 		ControlRigCentreOfMassNudge,
+		ControlRigJointEnabled,
 
 		// -----<new versions can be added above this line>-------------------------------------------------
 		VersionPlusOne,
