Control Rig Physics: Expose option to create controls relative to the parent vs world. Also renames the UI for this.

diff --git a/Engine/Plugins/Experimental/ControlRigPhysics/Source/ControlRigPhysics/Private/RigPhysicsControlExecution.cpp b/Engine/Plugins/Experimental/ControlRigPhysics/Source/ControlRigPhysics/Private/RigPhysicsControlExecution.cpp
--- a/Engine/Plugins/Experimental/ControlRigPhysics/Source/ControlRigPhysics/Private/RigPhysicsControlExecution.cpp
+++ b/Engine/Plugins/Experimental/ControlRigPhysics/Source/ControlRigPhysics/Private/RigPhysicsControlExecution.cpp
@@ -38,6 +38,7 @@ FRigUnit_AddPhysicsControl_Execute()
 				ExecuteContext.Hierarchy->FindComponent(ControlComponentKey)))
 			{
 				Component->ParentBodyComponentKey = ParentBodyComponentKey;
+				Component->bUseParentBodyAsDefault = bUseParentBodyAsDefault;
 				Component->ChildBodyComponentKey = ChildBodyComponentKey;
 				Component->ControlData = ControlData;
 				Component->ControlMultiplier = ControlMultiplier;
diff --git a/Engine/Plugins/Experimental/ControlRigPhysics/Source/ControlRigPhysics/Public/RigPhysicsControlComponent.h b/Engine/Plugins/Experimental/ControlRigPhysics/Source/ControlRigPhysics/Public/RigPhysicsControlComponent.h
--- a/Engine/Plugins/Experimental/ControlRigPhysics/Source/ControlRigPhysics/Public/RigPhysicsControlComponent.h
+++ b/Engine/Plugins/Experimental/ControlRigPhysics/Source/ControlRigPhysics/Public/RigPhysicsControlComponent.h
@@ -39,7 +39,7 @@ public:
 	// If true, then if the parent body component key is not set, the the default parent body comes
 	// from the parent joint. If it is false, then this search is not done, so the control will be
 	// in simulation space.
-	UPROPERTY(BlueprintReadOnly, EditAnywhere, Category = Control)
+	UPROPERTY(BlueprintReadOnly, DisplayName = "Use Parent Body", EditAnywhere, Category = Control)
 	bool bUseParentBodyAsDefault = false;
 
 	// The body being controlled
diff --git a/Engine/Plugins/Experimental/ControlRigPhysics/Source/ControlRigPhysics/Public/RigPhysicsControlExecution.h b/Engine/Plugins/Experimental/ControlRigPhysics/Source/ControlRigPhysics/Public/RigPhysicsControlExecution.h
--- a/Engine/Plugins/Experimental/ControlRigPhysics/Source/ControlRigPhysics/Public/RigPhysicsControlExecution.h
+++ b/Engine/Plugins/Experimental/ControlRigPhysics/Source/ControlRigPhysics/Public/RigPhysicsControlExecution.h
@@ -46,6 +46,9 @@ struct FRigUnit_AddPhysicsControl : public FRigUnit_PhysicsBaseMutable
 	UPROPERTY(meta = (Input))
 	FRigComponentKey ParentBodyComponentKey;
 
+	UPROPERTY(meta = (Input), DisplayName = "Use Parent Body")
+	bool bUseParentBodyAsDefault = false;
+
 	// The body that is controlled
 	UPROPERTY(meta = (Input))
 	FRigComponentKey ChildBodyComponentKey;
