Minor fixes

diff --git a/Engine/Plugins/Experimental/ControlRigPhysics/Source/ControlRigPhysics/Private/RigPhysicsSimulation.cpp b/Engine/Plugins/Experimental/ControlRigPhysics/Source/ControlRigPhysics/Private/RigPhysicsSimulation.cpp
--- a/Engine/Plugins/Experimental/ControlRigPhysics/Source/ControlRigPhysics/Private/RigPhysicsSimulation.cpp
+++ b/Engine/Plugins/Experimental/ControlRigPhysics/Source/ControlRigPhysics/Private/RigPhysicsSimulation.cpp
@@ -1045,6 +1045,7 @@ void FRigPhysicsSimulation::UpdateBodyPrePhysics(
 				}
 				default:
 				{
+					KinematicTargetCS = Record.SourceComponentSpaceTM.ToTransform();
 					UE_CONTROLRIG_RIGUNIT_REPORT_ERROR(TEXT("Kinematic target space is not valid"));
 				}
 				}
diff --git a/Engine/Plugins/Experimental/PhysicsControl/Source/PhysicsControl/Public/PhysicsControlData.h b/Engine/Plugins/Experimental/PhysicsControl/Source/PhysicsControl/Public/PhysicsControlData.h
--- a/Engine/Plugins/Experimental/PhysicsControl/Source/PhysicsControl/Public/PhysicsControlData.h
+++ b/Engine/Plugins/Experimental/PhysicsControl/Source/PhysicsControl/Public/PhysicsControlData.h
@@ -115,7 +115,7 @@ struct FPhysicsControlSetUpdates
 };
 
 // Specifies what space a kinematic target is defined in.
-UENUM()
+UENUM(BlueprintType)
 enum class EPhysicsControlKinematicTargetSpace : uint8
 {
 	// Target is in world space. For a control rig, this is the space of the control rig, rather
