Control Rig Physics: Initialise data - fixes compiler warnings

diff --git a/Engine/Plugins/Experimental/ControlRigPhysics/Source/ControlRigPhysics/Public/RigPhysicsControlExecution.h b/Engine/Plugins/Experimental/ControlRigPhysics/Source/ControlRigPhysics/Public/RigPhysicsControlExecution.h
--- a/Engine/Plugins/Experimental/ControlRigPhysics/Source/ControlRigPhysics/Public/RigPhysicsControlExecution.h
+++ b/Engine/Plugins/Experimental/ControlRigPhysics/Source/ControlRigPhysics/Public/RigPhysicsControlExecution.h
@@ -160,7 +160,7 @@ struct FRigUnit_HierarchySetControlLinearStrength : public FRigUnit_PhysicsBaseM
 	FRigComponentKey PhysicsControlComponentKey;
 
 	UPROPERTY(meta = (Input))
-	float Strength;
+	float Strength = 0.0f;
 };
 
 USTRUCT(meta = (DisplayName = "Set Physics Control Linear Damping Ratio", Varying))
@@ -181,7 +181,7 @@ struct FRigUnit_HierarchySetControlLinearDampingRatio : public FRigUnit_PhysicsB
 	FRigComponentKey PhysicsControlComponentKey;
 
 	UPROPERTY(meta = (Input))
-	float DampingRatio;
+	float DampingRatio = 1.0f;
 };
 
 USTRUCT(meta = (DisplayName = "Set Physics Control Angular Strength", Varying))
@@ -202,7 +202,7 @@ struct FRigUnit_HierarchySetControlAngularStrength : public FRigUnit_PhysicsBase
 	FRigComponentKey PhysicsControlComponentKey;
 
 	UPROPERTY(meta = (Input))
-	float Strength;
+	float Strength = 0.0f;
 };
 
 USTRUCT(meta = (DisplayName = "Set Physics Control Angular Damping Ratio", Varying))
@@ -223,7 +223,7 @@ struct FRigUnit_HierarchySetControlAngularDampingRatio : public FRigUnit_Physics
 	FRigComponentKey PhysicsControlComponentKey;
 
 	UPROPERTY(meta = (Input))
-	float DampingRatio;
+	float DampingRatio = 1.0f;
 };
 
 // Gets the control data for a physics control
diff --git a/Engine/Plugins/Experimental/ControlRigPhysics/Source/ControlRigPhysics/Public/RigPhysicsJointExecution.h b/Engine/Plugins/Experimental/ControlRigPhysics/Source/ControlRigPhysics/Public/RigPhysicsJointExecution.h
--- a/Engine/Plugins/Experimental/ControlRigPhysics/Source/ControlRigPhysics/Public/RigPhysicsJointExecution.h
+++ b/Engine/Plugins/Experimental/ControlRigPhysics/Source/ControlRigPhysics/Public/RigPhysicsJointExecution.h
@@ -124,7 +124,7 @@ struct FRigUnit_HierarchySetJointEnabled : public FRigUnit_PhysicsBaseMutable
 	FRigComponentKey PhysicsJointComponentKey;
 
 	UPROPERTY(meta = (Input))
-	bool bEnabled;
+	bool bEnabled = true;
 };
 
 // Sets the joint drive for a physics component body
