Require documentation for all the Control Rig Physics nodes, and add it where missing.

diff --git a/Engine/Plugins/Experimental/ControlRigPhysics/Source/ControlRigPhysics/Private/RigPhysicsExecution.cpp b/Engine/Plugins/Experimental/ControlRigPhysics/Source/ControlRigPhysics/Private/RigPhysicsExecution.cpp
--- a/Engine/Plugins/Experimental/ControlRigPhysics/Source/ControlRigPhysics/Private/RigPhysicsExecution.cpp
+++ b/Engine/Plugins/Experimental/ControlRigPhysics/Source/ControlRigPhysics/Private/RigPhysicsExecution.cpp
@@ -158,6 +158,22 @@ FRigUnit_SetPhysicsSolverSimulationSpaceSettings_Execute()
 	}
 }
 
+//======================================================================================================================
+FRigUnit_GetPhysicsSolverSimulationSpaceSettings_Execute()
+{
+	DECLARE_SCOPE_HIERARCHICAL_COUNTER_RIGUNIT();
+	if (!ExecuteContext.Hierarchy)
+	{
+		return;
+	}
+	URigHierarchy& Hierarchy = *ExecuteContext.Hierarchy;
+	if (const FRigPhysicsSolverComponent* SolverComponent = Cast<FRigPhysicsSolverComponent>(
+		Hierarchy.FindComponent(PhysicsSolverComponentKey)))
+	{
+		SimulationSpaceSettings = SolverComponent->SimulationSpaceSettings;
+	}
+}
+
 //======================================================================================================================
 FRigUnit_SetPhysicsSolverExternalVelocity_Execute()
 {
@@ -455,6 +471,12 @@ FRigUnit_HierarchyInstantiateFromPhysicsAsset_Execute()
 
 		if (bAddParentSpaceControl)
 		{
+			// Danny TODO Note that we create a parent-space control for every body - even the root
+			// body. This is often not desired, so we should provide a way of disabling it, as
+			// otherwise it will provide slightly unexpected behaviour (even though it's "correct")
+			// as we end up with a sim-space control if there is no parent body. We can't simply
+			// move this to the joint creation loop, since there may be some joints where we don't
+			// want a parent-space control (e.g. hamstrings).
 			FRigComponentKey ParentSpaceControlComponentKey = Controller->AddComponent(
 				FRigPhysicsControlComponent::StaticStruct(),
 				TEXT("ParentSpaceControl"), PhysicsBodyComponentKey.ElementKey);
@@ -691,7 +713,6 @@ FRigUnit_HierarchyImportCollisionFromPhysicsAsset_Execute()
 			}
 		}
 	}
-
 }
 
 //======================================================================================================================
diff --git a/Engine/Plugins/Experimental/ControlRigPhysics/Source/ControlRigPhysics/Public/RigPhysicsBodyExecution.h b/Engine/Plugins/Experimental/ControlRigPhysics/Source/ControlRigPhysics/Public/RigPhysicsBodyExecution.h
--- a/Engine/Plugins/Experimental/ControlRigPhysics/Source/ControlRigPhysics/Public/RigPhysicsBodyExecution.h
+++ b/Engine/Plugins/Experimental/ControlRigPhysics/Source/ControlRigPhysics/Public/RigPhysicsBodyExecution.h
@@ -33,18 +33,15 @@ struct FRigUnit_AddPhysicsBody : public FRigUnit_PhysicsBaseMutable
 	RIGVM_METHOD()
 	UE_API virtual void Execute() override;
 
-	/**
-	 * The owner of the newly created component (must be set/valid)
-	 */
+	// The owner of the newly created component (must be set/valid)
 	UPROPERTY(meta = (Input, BoneName))
 	FRigElementKey Owner;
 
+	// The Physics Body component key that was created
 	UPROPERTY(meta = (Output))
 	FRigComponentKey PhysicsBodyComponentKey;
 
-	/*
-	 * The solver to relate this new physics element to
-	 */
+	// The solver to relate this new physics body should be added to
 	UPROPERTY(meta = (Input))
 	FRigPhysicsBodySolverSettings Solver;
 
@@ -77,17 +74,18 @@ struct FRigUnit_HierarchyAutoCalculateCollision : public FRigUnit_PhysicsBaseMut
 	RIGVM_METHOD()
 	UE_API virtual void Execute() override;
 
+	// The Physics Body to be updated
 	UPROPERTY(meta = (Input))
 	FRigComponentKey PhysicsBodyComponentKey;
 
 	// For boxes: The minimum box extent, as a proportion of the maximum box extent.
 	// For capsules: The minimum radius, as a proportion of the length (not including the radius)
-	UPROPERTY(meta = (Input))
+	UPROPERTY(meta = (Input, ClampMin = "0.0"))
 	float MinAspectRatio = 0.25f;
 
 	// For boxes: The minimum side length. 
 	// For capsules: The minimum radius
-	UPROPERTY(meta = (Input))
+	UPROPERTY(meta = (Input, ClampMin = "0.0"))
 	float MinSize = 0.0f;
 };
 
@@ -106,9 +104,11 @@ struct FRigUnit_HierarchySetDynamics : public FRigUnit_PhysicsBaseMutable
 	RIGVM_METHOD()
 	UE_API virtual void Execute() override;
 
+	// The Physics Body to be updated
 	UPROPERTY(meta = (Input))
 	FRigComponentKey PhysicsBodyComponentKey;
 
+	// Sets the properties describing the dynamics properties of the Physics Body
 	UPROPERTY(meta = (Input))
 	FRigPhysicsDynamics Dynamics;
 };
@@ -128,9 +128,11 @@ struct FRigUnit_HierarchySetCollision : public FRigUnit_PhysicsBaseMutable
 	RIGVM_METHOD()
 	UE_API virtual void Execute() override;
 
+	// The Physics Body to be updated
 	UPROPERTY(meta = (Input))
 	FRigComponentKey PhysicsBodyComponentKey;
 
+	// Specifies the collision associated with the Physics Body
 	UPROPERTY(meta = (Input))
 	FRigPhysicsCollision Collision;
 };
@@ -153,9 +155,11 @@ struct FRigUnit_HierarchyDisableCollisionBetween : public FRigUnit_PhysicsBaseMu
 	RIGVM_METHOD()
 	UE_API virtual void Execute() override;
 
+	// Component key for the first body
 	UPROPERTY(meta = (Input))
 	FRigComponentKey PhysicsBodyComponentKey1;
 
+	// Component key for the second body
 	UPROPERTY(meta = (Input))
 	FRigComponentKey PhysicsBodyComponentKey2;
 };
@@ -179,9 +183,11 @@ struct FRigUnit_HierarchySetPhysicsBodySourceBone : public FRigUnit_PhysicsBaseM
 	RIGVM_METHOD()
 	UE_API virtual void Execute() override;
 
+	// The Physics Body to be updated
 	UPROPERTY(meta = (Input))
 	FRigComponentKey PhysicsBodyComponentKey;
 
+	// The bone to use as a source transform (kinematic target, and initialization) for the Physics Body.
 	UPROPERTY(meta = (Input))
 	FRigElementKey SourceBone;
 };
@@ -204,14 +210,15 @@ struct FRigUnit_HierarchySetPhysicsBodyTargetBone : public FRigUnit_PhysicsBaseM
 	RIGVM_METHOD()
 	UE_API virtual void Execute() override;
 
+	// The Physics Body to be updated
 	UPROPERTY(meta = (Input))
 	FRigComponentKey PhysicsBodyComponentKey;
 
+	// The bone to use as a target for the Physics Body - i.e. where its transforms should be written to.
 	UPROPERTY(meta = (Input))
 	FRigElementKey TargetBone;
 };
 
-
 /**
  * Sets all the data on a body - but in a sparse way so you can decide which parameters get applied.
  * Danny TODO - Note that the sparse data does not get displayed correctly, so this is largely 
@@ -231,9 +238,11 @@ struct FRigUnit_HierarchySetPhysicsBodySparseData : public FRigUnit_PhysicsBaseM
 	RIGVM_METHOD()
 	UE_API virtual void Execute() override;
 
+	// The Physics Body to be updated
 	UPROPERTY(meta = (Input))
 	FRigComponentKey PhysicsBodyComponentKey;
 
+	// Sparse data to be set on the Physics Body
 	UPROPERTY(meta = (Input))
 	FPhysicsControlModifierSparseData Data;
 };
@@ -255,9 +264,12 @@ struct FRigUnit_HierarchySetPhysicsBodyKinematicTarget : public FRigUnit_Physics
 	RIGVM_METHOD()
 	UE_API virtual void Execute() override;
 
+	// The Physics Body to be updated
 	UPROPERTY(meta = (Input))
 	FRigComponentKey PhysicsBodyComponentKey;
 
+	// Kinematic target for the Physics Body. This will be combined with the animation pose,
+	// depending on the Kinematic Target Space.
 	UPROPERTY(meta = (Input))
 	FTransform KinematicTarget;
 };
@@ -279,16 +291,17 @@ struct FRigUnit_HierarchySetPhysicsBodyKinematicTargetSpace : public FRigUnit_Ph
 	RIGVM_METHOD()
 	UE_API virtual void Execute() override;
 
+	// The Physics Body to be updated
 	UPROPERTY(meta = (Input))
 	FRigComponentKey PhysicsBodyComponentKey;
 
+	// The space for kinematic targets.
 	UPROPERTY(meta = (Input))
 	EPhysicsControlKinematicTargetSpace KinematicTargetSpace = EPhysicsControlKinematicTargetSpace::OffsetInBoneSpace;
 };
 
 /**
  * Sets the movement mode for this body. 
- * Danny TODO explain the different between kinematic and static
  */
 USTRUCT(meta = (DisplayName = "Set Physics Body Movement Mode", Varying))
 struct FRigUnit_HierarchySetPhysicsBodyMovementType : public FRigUnit_PhysicsBaseMutable
@@ -304,9 +317,11 @@ struct FRigUnit_HierarchySetPhysicsBodyMovementType : public FRigUnit_PhysicsBas
 	RIGVM_METHOD()
 	UE_API virtual void Execute() override;
 
+	// The Physics Body to be updated
 	UPROPERTY(meta = (Input))
 	FRigComponentKey PhysicsBodyComponentKey;
 
+	// How the Physics Body should move when the Physics Solver is stepped.
 	UPROPERTY(meta = (Input))
 	EPhysicsMovementType MovementType = EPhysicsMovementType::Simulated;
 };
@@ -328,14 +343,15 @@ struct FRigUnit_HierarchySetPhysicsBodyCollisionType : public FRigUnit_PhysicsBa
 	RIGVM_METHOD()
 	UE_API virtual void Execute() override;
 
+	// The Physics Body to be updated
 	UPROPERTY(meta = (Input))
 	FRigComponentKey PhysicsBodyComponentKey;
 
+	// What type of collision to use for the Physics Body
 	UPROPERTY(meta = (Input))
 	TEnumAsByte<ECollisionEnabled::Type> CollisionType = ECollisionEnabled::QueryAndPhysics;
 };
 
-
 // Sets whether this body should be included in checks for resetting physics on the whole rig
 USTRUCT(meta = (DisplayName = "Set Physics Body Include In Checks For Reset", Varying))
 struct FRigUnit_HierarchySetPhysicsBodyIncludeInChecksForReset : public FRigUnit_PhysicsBaseMutable
@@ -351,9 +367,11 @@ struct FRigUnit_HierarchySetPhysicsBodyIncludeInChecksForReset : public FRigUnit
 	RIGVM_METHOD()
 	UE_API virtual void Execute() override;
 
+	// The Physics Body to be updated
 	UPROPERTY(meta = (Input))
 	FRigComponentKey PhysicsBodyComponentKey;
 
+	// WHether the Physics Body should be included when checking to see if the simulation should be reset.
 	UPROPERTY(meta = (Input))
 	bool bInclude = true;
 };
@@ -373,9 +391,11 @@ struct FRigUnit_HierarchySetPhysicsBodyMaterial : public FRigUnit_PhysicsBaseMut
 	RIGVM_METHOD()
 	UE_API virtual void Execute() override;
 
+	// The Physics Body to be updated
 	UPROPERTY(meta = (Input))
 	FRigComponentKey PhysicsBodyComponentKey;
 
+	// The physics material to be used for collision on the Physics Body
 	UPROPERTY(meta = (Input))
 	FRigPhysicsMaterial Material;
 };
@@ -397,16 +417,19 @@ struct FRigUnit_HierarchySetPhysicsBodyGravityMultiplier : public FRigUnit_Physi
 	RIGVM_METHOD()
 	UE_API virtual void Execute() override;
 
+	// The Physics Body to be updated
 	UPROPERTY(meta = (Input))
 	FRigComponentKey PhysicsBodyComponentKey;
 
+	// How much the Physics Body should respond to the gravity set in the Physics Solver
 	UPROPERTY(meta = (Input))
 	float GravityMultiplier = 1.0f;
 };
 
 /**
  * Controls the amount that the simulation is blended back into the target bones. 
- * Danny TODO implement
+ * Danny TODO implement. Note that this isn't completely obvious - sometimes we want to blend 
+ * back in the hierarchy space, and sometimes in simulation space.
  */
 USTRUCT(meta = (DisplayName = "Set Physics Body Physics Blend Weight", Varying))
 struct FRigUnit_HierarchySetPhysicsBodyPhysicsBlendWeight : public FRigUnit_PhysicsBaseMutable
@@ -422,10 +445,12 @@ struct FRigUnit_HierarchySetPhysicsBodyPhysicsBlendWeight : public FRigUnit_Phys
 	RIGVM_METHOD()
 	UE_API virtual void Execute() override;
 
+	// The Physics Body to be updated
 	UPROPERTY(meta = (Input))
 	FRigComponentKey PhysicsBodyComponentKey;
 
-	UPROPERTY(meta = (Input))
+	// A per-body "alpha" on how much of the physics movement to apply to the target bone.
+	UPROPERTY(meta = (Input, ClampMin = "0.0", ClampMax = "1.0"))
 	float PhysicsBlendWeight = 1.0f;
 };
 
@@ -448,17 +473,20 @@ struct FRigUnit_HierarchySetPhysicsBodyUpdateKinematicFromSimulation : public FR
 	RIGVM_METHOD()
 	UE_API virtual void Execute() override;
 
+	// The Physics Body to be updated
 	UPROPERTY(meta = (Input))
 	FRigComponentKey PhysicsBodyComponentKey;
 
+	// When bodies are written back from simulation, if this is set to false, there is a small
+	// performance benefit to skipping bodies that are known to be kinematic, and when they are set
+	// to just track the input animation.
 	UPROPERTY(meta = (Input))
 	bool bUpdateKinematicFromSimulation = true;
 };
 
 /**
- * If true, then kinematic objects will be written back from simulation to the bones. This only 
- * necessary when either kinematic targets are being used, or when the target bone differs from the source bone.
- * Danny TODO - check the implementation, and also can we auto-calculate this at runtime?
+ * Sets the linear and angular damping on the body. This will reduce the velocity, to make the 
+ * body start tracking the simulation space.
  */
 USTRUCT(meta = (DisplayName = "Set Physics Body Damping", Varying))
 struct FRigUnit_HierarchySetPhysicsBodyDamping : public FRigUnit_PhysicsBaseMutable
@@ -474,13 +502,16 @@ struct FRigUnit_HierarchySetPhysicsBodyDamping : public FRigUnit_PhysicsBaseMuta
 	RIGVM_METHOD()
 	UE_API virtual void Execute() override;
 
+	// The Physics Body to be updated
 	UPROPERTY(meta = (Input))
 	FRigComponentKey PhysicsBodyComponentKey;
 
-	UPROPERTY(meta = (Input))
+	// The desired linear damping
+	UPROPERTY(meta = (Input, ClampMin = "0.0"))
 	float LinearDamping = 0;
 
-	UPROPERTY(meta = (Input))
+	// The desired angular damping
+	UPROPERTY(meta = (Input, ClampMin = "0.0"))
 	float AngularDamping = 0;
 };
 
diff --git a/Engine/Plugins/Experimental/ControlRigPhysics/Source/ControlRigPhysics/Public/RigPhysicsControlExecution.h b/Engine/Plugins/Experimental/ControlRigPhysics/Source/ControlRigPhysics/Public/RigPhysicsControlExecution.h
--- a/Engine/Plugins/Experimental/ControlRigPhysics/Source/ControlRigPhysics/Public/RigPhysicsControlExecution.h
+++ b/Engine/Plugins/Experimental/ControlRigPhysics/Source/ControlRigPhysics/Public/RigPhysicsControlExecution.h
@@ -33,12 +33,11 @@ struct FRigUnit_AddPhysicsControl : public FRigUnit_PhysicsBaseMutable
 	RIGVM_METHOD()
 	UE_API virtual void Execute() override;
 
-	/**
-	 * The owner of the newly created component (must be set/valid)
-	 */
+	// The owner of the newly created component (must be set/valid)
 	UPROPERTY(meta = (Input, BoneName))
 	FRigElementKey Owner;
 
+	// The Physics Control component key that was created
 	UPROPERTY(meta = (Output))
 	FRigComponentKey ControlComponentKey;
 
@@ -46,6 +45,7 @@ struct FRigUnit_AddPhysicsControl : public FRigUnit_PhysicsBaseMutable
 	UPROPERTY(meta = (Input))
 	FRigComponentKey ParentBodyComponentKey;
 
+	// Whether or not to use the parent body (if one exists) as the Physics Control parent
 	UPROPERTY(meta = (Input), DisplayName = "Use Parent Body")
 	bool bUseParentBodyAsDefault = false;
 
@@ -53,19 +53,19 @@ struct FRigUnit_AddPhysicsControl : public FRigUnit_PhysicsBaseMutable
 	UPROPERTY(meta = (Input))
 	FRigComponentKey ChildBodyComponentKey;
 
-	/** Describes the initial strength etc of the new control */
+	// Describes the initial strength etc of the new control
 	UPROPERTY(meta = (Input))
 	FPhysicsControlData ControlData;
 
-	/** Fine control over the control strengths etc */
+	// Fine control over the control strengths etc
 	UPROPERTY(meta = (Input))
 	FPhysicsControlMultiplier ControlMultiplier;
 
-	/** Describes the initial target for the new control */
+	// The initial target for the new control
 	UPROPERTY(meta = (Input))
 	FPhysicsControlTarget ControlTarget;
 
-	/** Which set to include the control in (optional). Note that it automatically gets added to the set "All" */
+	// Danny TODO Which set to include the control in (optional). Note that it automatically gets added to the set "All"
 	//UPROPERTY(meta = (Input))
 	//FName Set;
 };
@@ -85,9 +85,11 @@ struct FRigUnit_HierarchySetControlEnabled : public FRigUnit_PhysicsBaseMutable
 	RIGVM_METHOD()
 	UE_API virtual void Execute() override;
 
+	// The Physics Control that should be enabled or disabled
 	UPROPERTY(meta = (Input))
 	FRigComponentKey PhysicsControlComponentKey;
 
+	// Whether or not the Physics Control should be enabled
 	UPROPERTY(meta = (Input))
 	bool bEnabled = true;
 };
@@ -107,6 +109,7 @@ struct FRigUnit_HierarchySetControlCustomControlPoint : public FRigUnit_PhysicsB
 	RIGVM_METHOD()
 	UE_API virtual void Execute() override;
 
+	// The Physics Control component to be updated
 	UPROPERTY(meta = (Input))
 	FRigComponentKey PhysicsControlComponentKey;
 
@@ -135,13 +138,16 @@ struct FRigUnit_HierarchySetControlData : public FRigUnit_PhysicsBaseMutable
 	RIGVM_METHOD()
 	UE_API virtual void Execute() override;
 
+	// The Physics Control component to be updated
 	UPROPERTY(meta = (Input))
 	FRigComponentKey PhysicsControlComponentKey;
 
+	// The data (strengths etc) that should be set on the Physics Control
 	UPROPERTY(meta = (Input))
 	FPhysicsControlData ControlData;
 };
 
+// Sets the Linear Strength of a Physics Control
 USTRUCT(meta = (DisplayName = "Set Physics Control Linear Strength", Varying))
 struct FRigUnit_HierarchySetControlLinearStrength : public FRigUnit_PhysicsBaseMutable
 {
@@ -156,13 +162,19 @@ struct FRigUnit_HierarchySetControlLinearStrength : public FRigUnit_PhysicsBaseM
 	RIGVM_METHOD()
 	UE_API virtual void Execute() override;
 
+	// The Physics Control component to be updated
 	UPROPERTY(meta = (Input))
 	FRigComponentKey PhysicsControlComponentKey;
 
-	UPROPERTY(meta = (Input))
+	// The linear strength of the control. In the absence of any other influences, this is the
+	// frequency of oscillation when there is zero damping. For example, a strength of 2 would make
+	// the target oscillate twice per second (when there is no damping).
+	// Note that parent-space controls should normally have this set to zero, since 
+	UPROPERTY(meta = (Input, ClampMin = "0.0"))
 	float Strength = 0.0f;
 };
 
+// Sets the Linear Damping Ratio of a Physics Control
 USTRUCT(meta = (DisplayName = "Set Physics Control Linear Damping Ratio", Varying))
 struct FRigUnit_HierarchySetControlLinearDampingRatio : public FRigUnit_PhysicsBaseMutable
 {
@@ -177,13 +189,19 @@ struct FRigUnit_HierarchySetControlLinearDampingRatio : public FRigUnit_PhysicsB
 	RIGVM_METHOD()
 	UE_API virtual void Execute() override;
 
+	// The Physics Control component to be updated
 	UPROPERTY(meta = (Input))
 	FRigComponentKey PhysicsControlComponentKey;
 
-	UPROPERTY(meta = (Input))
+	// The linear damping ratio for the control. When this is 1, there is just enough damping to
+	// stop the control from oscillating about its target, due to the control strength (in the
+	// absence of any other influences). Values above this will add more damping, and values below
+	// will make the system tend to oscillate, as it will be under-ramped.
+	UPROPERTY(meta = (Input, ClampMin = "0.0"))
 	float DampingRatio = 1.0f;
 };
 
+// Sets the Angular Strength of a Physics Control
 USTRUCT(meta = (DisplayName = "Set Physics Control Angular Strength", Varying))
 struct FRigUnit_HierarchySetControlAngularStrength : public FRigUnit_PhysicsBaseMutable
 {
@@ -198,13 +216,18 @@ struct FRigUnit_HierarchySetControlAngularStrength : public FRigUnit_PhysicsBase
 	RIGVM_METHOD()
 	UE_API virtual void Execute() override;
 
+	// The Physics Control component to be updated
 	UPROPERTY(meta = (Input))
 	FRigComponentKey PhysicsControlComponentKey;
 
-	UPROPERTY(meta = (Input))
+	// The angular strength of the control. In the absence of any other influences, this is the
+	// frequency of oscillation when there is zero damping. For example, a strength of 2 would make
+	// the target oscillate twice per second (when there is no damping).
+	UPROPERTY(meta = (Input, ClampMin = "0.0"))
 	float Strength = 0.0f;
 };
 
+// Sets the Angular Damping Ratio of a Physics Control
 USTRUCT(meta = (DisplayName = "Set Physics Control Angular Damping Ratio", Varying))
 struct FRigUnit_HierarchySetControlAngularDampingRatio : public FRigUnit_PhysicsBaseMutable
 {
@@ -219,10 +242,15 @@ struct FRigUnit_HierarchySetControlAngularDampingRatio : public FRigUnit_Physics
 	RIGVM_METHOD()
 	UE_API virtual void Execute() override;
 
+	// The Physics Control component to be updated
 	UPROPERTY(meta = (Input))
 	FRigComponentKey PhysicsControlComponentKey;
 
-	UPROPERTY(meta = (Input))
+	// The angular damping ratio for the control. When this is 1, there is just enough damping to
+	// stop the control from oscillating about its target, due to the control strength (in the
+	// absence of any other influences). Values above this will add more damping, and values below
+	// will make the system tend to oscillate, as it will be under-ramped.
+	UPROPERTY(meta = (Input, ClampMin = "0.0"))
 	float DampingRatio = 1.0f;
 };
 
@@ -241,9 +269,11 @@ struct FRigUnit_HierarchyGetControlData : public FRigUnit_PhysicsBase
 	RIGVM_METHOD()
 	UE_API virtual void Execute() override;
 
+	// The Physics Control component from which to retrieve data
 	UPROPERTY(meta = (Input))
 	FRigComponentKey PhysicsControlComponentKey;
 
+	// The current control data (strengths etc) 
 	UPROPERTY(meta = (Output))
 	FPhysicsControlData ControlData;
 };
@@ -263,9 +293,12 @@ struct FRigUnit_HierarchySetControlMultiplier : public FRigUnit_PhysicsBaseMutab
 	RIGVM_METHOD()
 	UE_API virtual void Execute() override;
 
+	// The Physics Control component to be updated
 	UPROPERTY(meta = (Input))
 	FRigComponentKey PhysicsControlComponentKey;
 
+	// The multipliers to use on the control. These allow fine-grained control over the strengths
+	// (etc) in different directions.
 	UPROPERTY(meta = (Input))
 	FPhysicsControlMultiplier ControlMultiplier;
 };
@@ -285,12 +318,15 @@ struct FRigUnit_HierarchySetControlDataAndMultiplier : public FRigUnit_PhysicsBa
 	RIGVM_METHOD()
 	UE_API virtual void Execute() override;
 
+	// The Physics Control component to be updated
 	UPROPERTY(meta = (Input))
 	FRigComponentKey PhysicsControlComponentKey;
 
+	// The control data (strengths etc) to use
 	UPROPERTY(meta = (Input))
 	FPhysicsControlData ControlData;
 
+	// Detail/per-direction multipliers on the control data
 	UPROPERTY(meta = (Input))
 	FPhysicsControlMultiplier ControlMultiplier;
 };
@@ -310,9 +346,11 @@ struct FRigUnit_HierarchySetControlTarget : public FRigUnit_PhysicsBaseMutable
 	RIGVM_METHOD()
 	UE_API virtual void Execute() override;
 
+	// The Physics Control component to be updated
 	UPROPERTY(meta = (Input))
 	FRigComponentKey PhysicsControlComponentKey;
 
+	// The target transform (etc) for the control
 	UPROPERTY(meta = (Input))
 	FPhysicsControlTarget ControlTarget;
 };
@@ -333,19 +371,20 @@ struct FRigUnit_HierarchyUpdateControlTarget : public FRigUnit_PhysicsBaseMutabl
 	RIGVM_METHOD()
 	UE_API virtual void Execute() override;
 
+	// The Physics Control component to be updated
 	UPROPERTY(meta = (Input))
 	FRigComponentKey PhysicsControlComponentKey;
 
-	/** The target position of the child body, relative to the parent body */
+	// The target position of the child body, relative to the parent body
 	UPROPERTY(meta = (Input))
 	FVector TargetPosition = FVector::ZeroVector;
 
-	/** The target orientation of the child body, relative to the parent body */
+	// The target orientation of the child body, relative to the parent body
 	UPROPERTY(meta = (Input))
 	FRotator TargetOrientation = FRotator::ZeroRotator;
 
 	// The delta time used to calculate the target velocity
-	UPROPERTY(meta = (Input))
+	UPROPERTY(meta = (Input, ClampMin = "0.0"))
 	float DeltaTime = 0.0f;
 };
 
diff --git a/Engine/Plugins/Experimental/ControlRigPhysics/Source/ControlRigPhysics/Public/RigPhysicsExecution.h b/Engine/Plugins/Experimental/ControlRigPhysics/Source/ControlRigPhysics/Public/RigPhysicsExecution.h
--- a/Engine/Plugins/Experimental/ControlRigPhysics/Source/ControlRigPhysics/Public/RigPhysicsExecution.h
+++ b/Engine/Plugins/Experimental/ControlRigPhysics/Source/ControlRigPhysics/Public/RigPhysicsExecution.h
@@ -16,7 +16,7 @@
 /**
  * Base struct for all other mutable physics nodes
  */
-USTRUCT(meta = (Category = "RigPhysics", NodeColor = "1.0 0.6 0.3", Keywords = "Physics"))
+USTRUCT(meta = (Category = "RigPhysics", NodeColor = "1.0 0.6 0.3", Keywords = "Physics", DocumentationPolicy = "Strict"))
 struct FRigUnit_PhysicsBaseMutable: public FRigUnitMutable
 {
 	GENERATED_BODY()
@@ -25,7 +25,7 @@ struct FRigUnit_PhysicsBaseMutable: public FRigUnitMutable
 /**
  * Base struct for all other non-mutable physics nodes
  */
-USTRUCT(meta = (Category = "RigPhysics", NodeColor = "1.0 0.6 0.3", Keywords = "Physics"))
+USTRUCT(meta = (Category = "RigPhysics", NodeColor = "1.0 0.6 0.3", Keywords = "Physics", DocumentationPolicy = "Strict"))
 struct FRigUnit_PhysicsBase : public FRigUnit
 {
 	GENERATED_BODY()
@@ -52,18 +52,21 @@ struct FRigUnit_AddPhysicsSolver : public FRigUnit_PhysicsBaseMutable
 	RIGVM_METHOD()
 	UE_API virtual void Execute() override;
 
-	/**
-	 * The owner of the newly created component (must be set/valid)
-	 */
+	// The owner of the newly created component (must be set/valid)
 	UPROPERTY(meta = (Input, BoneName))
 	FRigElementKey Owner;
 
+	// The component key of the solver that has been added
 	UPROPERTY(meta = (Output))
 	FRigComponentKey PhysicsSolverComponentKey;
 
+	// Settings for the physics solver that will be added
 	UPROPERTY(meta = (Input))
 	FRigPhysicsSolverSettings SolverSettings;
 
+	// Settings for the solver that apply to when it uses a simulation space other than "world".
+	// These typically relate to the movement of the simulation space itself, and how that is used
+	// by the solver.
 	UPROPERTY(meta = (Input))
 	FRigPhysicsSimulationSpaceSettings SimulationSpaceSettings;
 };
@@ -88,14 +91,11 @@ struct FRigUnit_InstantiatePhysics : public FRigUnit_PhysicsBaseMutable
 	RIGVM_METHOD()
 	UE_API virtual void Execute() override;
 
-	/*
-	 * The solver to relate this new physics element to
-	 */
+	// The solver that should be instantiated
 	UPROPERTY(meta = (Input))
 	FRigComponentKey PhysicsSolverComponentKey;
 };
 
-
 /**
  * Steps the specified physics solver
  */
@@ -113,7 +113,7 @@ struct FRigUnit_StepPhysicsSolver : public FRigUnit_PhysicsBaseMutable
 	RIGVM_METHOD()
 	UE_API virtual void Execute() override;
 
-	// The physics solver
+	// The physics solver that should be stepped
 	UPROPERTY(meta = (Input))
 	FRigComponentKey PhysicsSolverComponentKey;
 	
@@ -135,11 +135,12 @@ struct FRigUnit_StepPhysicsSolver : public FRigUnit_PhysicsBaseMutable
 	UPROPERTY(meta = (Input, ClampMin = "0.0", ClampMax = "1.0"))
 	float Alpha = 1.0f;
 
+	// Settings that specify how the solver state should be visualized during/after the step
 	UPROPERTY(meta = (Input))
 	FRigPhysicsVisualizationSettings VisualizationSettings;
 };
 
-// Sets all the simulation space settings on the solver
+// Sets the solver's simulation space settings
 USTRUCT(meta = (DisplayName = "Set Physics Solver Simulation Space Settings"))
 struct FRigUnit_SetPhysicsSolverSimulationSpaceSettings : public FRigUnit_PhysicsBaseMutable
 {
@@ -163,6 +164,30 @@ struct FRigUnit_SetPhysicsSolverSimulationSpaceSettings : public FRigUnit_Physic
 	FRigPhysicsSimulationSpaceSettings SimulationSpaceSettings;
 };
 
+// Gets the solver's simulation space settings
+USTRUCT(meta = (DisplayName = "Get Physics Solver Simulation Space Settings"))
+struct FRigUnit_GetPhysicsSolverSimulationSpaceSettings : public FRigUnit_PhysicsBase
+{
+	GENERATED_BODY()
+
+	FRigUnit_GetPhysicsSolverSimulationSpaceSettings()
+	{
+		PhysicsSolverComponentKey.ElementKey.Type = ERigElementType::Bone;
+		PhysicsSolverComponentKey.Name = FRigPhysicsSolverComponent::GetDefaultName();
+	}
+
+	RIGVM_METHOD()
+	UE_API virtual void Execute() override;
+
+	// The physics solver
+	UPROPERTY(meta = (Input))
+	FRigComponentKey PhysicsSolverComponentKey;
+
+	// The current simulation space settings
+	UPROPERTY(meta = (Output))
+	FRigPhysicsSimulationSpaceSettings SimulationSpaceSettings;
+};
+
 // Sets the external velocity of the simulation - used for adding wind effects
 USTRUCT(meta = (DisplayName = "Set Physics Solver External Velocity"))
 struct FRigUnit_SetPhysicsSolverExternalVelocity : public FRigUnit_PhysicsBaseMutable
@@ -250,30 +275,35 @@ struct FRigUnit_AddPhysicsComponents : public FRigUnit_PhysicsBaseMutable
 	RIGVM_METHOD()
 	UE_API virtual void Execute() override;
 
-	/**
-	 * The owner of the newly created component (must be set/valid)
-	 */
+	// The owner of the newly created component (must be set/valid)
 	UPROPERTY(meta = (Input, BoneName))
 	FRigElementKey Owner;
 
+	// Whether to add a Physics Joint
 	UPROPERTY(meta = (Input))
 	bool bAddJoint = true;
 
+	// Whether to add a simulation-space Physics Control
 	UPROPERTY(meta = (Input))
 	bool bAddSimSpaceControl = true;
 
+	// Whether to add a parent-space Physics Control
 	UPROPERTY(meta = (Input))
 	bool bAddParentSpaceControl = true;
 
+	// The component key of the Physics Body
 	UPROPERTY(meta = (Output))
 	FRigComponentKey PhysicsBodyComponentKey;
 
+	// The component key of the Physics Joint, if created
 	UPROPERTY(meta = (Output))
 	FRigComponentKey PhysicsJointComponentKey;
 
+	// The component key of the simulation-space Physics Control, if created
 	UPROPERTY(meta = (Output))
 	FRigComponentKey SimSpaceControlComponentKey;
 
+	// The component key of the parent-space Physics Control, if created
 	UPROPERTY(meta = (Output))
 	FRigComponentKey ParentSpaceControlComponentKey;
 
@@ -337,13 +367,10 @@ struct FRigUnit_HierarchyImportCollisionFromPhysicsAsset : public FRigUnit_Physi
 	UPROPERTY(meta = (Input))
 	bool bUseAutomaticSolver = true;
 
+	// The physics asset to import collision from
 	UPROPERTY(meta = (Input))
 	TObjectPtr<UPhysicsAsset> PhysicsAsset;
 
-	// Name of the constraint profile to use. If empty (or invalid), the default profile will be used
-	UPROPERTY(meta = (Input))
-	FName ConstraintProfileName;
-
 	// If this is empty, then all bones with bodies in the physics asset will be created. Otherwise
 	// only bodies that relate to the specified bones will be created.
 	UPROPERTY(meta = (Input))
@@ -357,9 +384,11 @@ struct FRigUnit_HierarchyImportCollisionFromPhysicsAsset : public FRigUnit_Physi
 	UPROPERTY(meta = (Input))
 	FRigElementKey Owner;
 
+	// The element keys of the bones that were created to own the physics bodies
 	UPROPERTY(meta = (Output))
 	TArray<FRigElementKey> BoneKeys;
 
+	// The Physics Body component keys that were created
 	UPROPERTY(meta = (Output))
 	TArray<FRigComponentKey> PhysicsBodyComponentKeys;
 };
@@ -391,6 +420,7 @@ struct FRigUnit_HierarchyInstantiateFromPhysicsAsset : public FRigUnit_PhysicsBa
 	UPROPERTY(meta = (Input))
 	bool bUseAutomaticSolver = true;
 
+	// The Physics Asset to instantiate from
 	UPROPERTY(meta = (Input))
 	TObjectPtr<UPhysicsAsset> PhysicsAsset;
 
@@ -413,9 +443,13 @@ struct FRigUnit_HierarchyInstantiateFromPhysicsAsset : public FRigUnit_PhysicsBa
 	UPROPERTY(meta = (Input))
 	bool bEnableDrives = true;
 
+	// Whether to create a simulation-space control for each body that was created
 	UPROPERTY(meta = (Input))
 	bool bAddSimSpaceControl = false;
 
+	// Whether to create a parent-space control for each body that was created. Note that if this is
+	// set to true, then it will be created for the root-most body in the physics asset too. This is
+	// often not needed, so make sure it isn't enabled if you don't want it to be.
 	UPROPERTY(meta = (Input))
 	bool bAddParentSpaceControl = false;
 
@@ -427,23 +461,25 @@ struct FRigUnit_HierarchyInstantiateFromPhysicsAsset : public FRigUnit_PhysicsBa
 	UPROPERTY(meta = (Input))
 	FPhysicsControlData ParentSpaceControlData;
 
+	// The Physics Body component keys that were created
 	UPROPERTY(meta = (Output))
 	TArray<FRigComponentKey> PhysicsBodyComponentKeys;
 
+	// The Physics Joint component keys that were created
 	UPROPERTY(meta = (Output))
 	TArray<FRigComponentKey> PhysicsJointComponentKeys;
 
+	// The simulation-space Physics Control component keys that were created
 	UPROPERTY(meta = (Output))
 	TArray<FRigComponentKey> SimSpaceControlComponentKeys;
 
+	// The parent-space Physics Control component keys that were created
 	UPROPERTY(meta = (Output))
 	TArray<FRigComponentKey> ParentSpaceControlComponentKeys;
 };
 
-
-// Retrieves the simulation space data. Note that this will have been generated during the
-// simulation step, so the values returned will relate to the previous update if the solver has not
-// yet been stepped.
+// Retrieves the simulation space data that were generated during the simulation step, so the values
+// returned will relate to the previous update if the solver has not yet been stepped.
 USTRUCT(meta = (DisplayName = "Get Physics Solver Space Data", Keywords = "Debug"))
 struct FRigUnit_GetPhysicsSolverSpaceData : public FRigUnit_PhysicsBase
 {
@@ -462,18 +498,23 @@ struct FRigUnit_GetPhysicsSolverSpaceData : public FRigUnit_PhysicsBase
 	UPROPERTY(meta = (Input))
 	FRigComponentKey PhysicsSolverComponentKey;
 
+	// The velocity of the simulation space (in world space)
 	UPROPERTY(meta = (Output))
 	FVector LinearVelocity = FVector::ZeroVector;
 
+	// The angular velocity of the simulation space (in world space)
 	UPROPERTY(meta = (Output))
 	FVector AngularVelocity = FVector::ZeroVector;
 
+	// The linear acceleration of the simulation space (in world space)
 	UPROPERTY(meta = (Output))
 	FVector LinearAcceleration = FVector::ZeroVector;
 
+	// The angular acceleration of the simulation space (in world space)
 	UPROPERTY(meta = (Output))
 	FVector AngularAcceleration = FVector::ZeroVector;
 
+	// The gravitational acceleration that will be applied (in simulation space)
 	UPROPERTY(meta = (Output))
 	FVector Gravity = FVector::ZeroVector;
 };
diff --git a/Engine/Plugins/Experimental/ControlRigPhysics/Source/ControlRigPhysics/Public/RigPhysicsJointExecution.h b/Engine/Plugins/Experimental/ControlRigPhysics/Source/ControlRigPhysics/Public/RigPhysicsJointExecution.h
--- a/Engine/Plugins/Experimental/ControlRigPhysics/Source/ControlRigPhysics/Public/RigPhysicsJointExecution.h
+++ b/Engine/Plugins/Experimental/ControlRigPhysics/Source/ControlRigPhysics/Public/RigPhysicsJointExecution.h
@@ -35,12 +35,11 @@ struct FRigUnit_AddPhysicsJoint : public FRigUnit_PhysicsBaseMutable
 	RIGVM_METHOD()
 	UE_API virtual void Execute() override;
 
-	/**
-	 * The owner of the newly created component (must be set/valid)
-	 */
+	// The owner of the newly created component (must be set/valid)
 	UPROPERTY(meta = (Input, BoneName))
 	FRigElementKey Owner;
 
+	// The Physics Joint component key that was created
 	UPROPERTY(meta = (Output))
 	FRigComponentKey PhysicsJointComponentKey;
 
@@ -77,9 +76,11 @@ struct FRigUnit_HierarchySetJointData : public FRigUnit_PhysicsBaseMutable
 	RIGVM_METHOD()
 	UE_API virtual void Execute() override;
 
+	// The Physics Joint to be updated
 	UPROPERTY(meta = (Input))
 	FRigComponentKey PhysicsJointComponentKey;
 
+	// The physics joint data (limit and drive properties) to be used
 	UPROPERTY(meta = (Input))
 	FRigPhysicsJointData JointData;
 };
@@ -99,13 +100,16 @@ struct FRigUnit_HierarchyGetJointData : public FRigUnit_PhysicsBase
 	RIGVM_METHOD()
 	UE_API virtual void Execute() override;
 
+	// The Physics Joint from which to retrieve data
 	UPROPERTY(meta = (Input))
 	FRigComponentKey PhysicsJointComponentKey;
 
+	// The current Physics Joint data (state of limits and drives)
 	UPROPERTY(meta = (Output))
 	FRigPhysicsJointData JointData;
 };
 
+// Specifies whether a Physics Joint should be enabled or not
 USTRUCT(meta = (DisplayName = "Set Physics Joint Enabled", Varying))
 struct FRigUnit_HierarchySetJointEnabled : public FRigUnit_PhysicsBaseMutable
 {
@@ -120,9 +124,11 @@ struct FRigUnit_HierarchySetJointEnabled : public FRigUnit_PhysicsBaseMutable
 	RIGVM_METHOD()
 	UE_API virtual void Execute() override;
 
+	// The Physics Joint to be updated
 	UPROPERTY(meta = (Input))
 	FRigComponentKey PhysicsJointComponentKey;
 
+	// Whether or not the joint should be enabled. Disabling the joint will disable any limits and drives
 	UPROPERTY(meta = (Input))
 	bool bEnabled = true;
 };
@@ -142,9 +148,11 @@ struct FRigUnit_HierarchySetJointDriveData : public FRigUnit_PhysicsBaseMutable
 	RIGVM_METHOD()
 	UE_API virtual void Execute() override;
 
+	// The Physics Joint to be updated
 	UPROPERTY(meta = (Input))
 	FRigComponentKey PhysicsJointComponentKey;
 
+	// The Physics Joint drive data that should be used
 	UPROPERTY(meta = (Input))
 	FRigPhysicsDriveData DriveData;
 };
@@ -164,9 +172,11 @@ struct FRigUnit_HierarchyGetJointDriveData : public FRigUnit_PhysicsBase
 	RIGVM_METHOD()
 	UE_API virtual void Execute() override;
 
+	// The Physics Joint from which to retrieve data
 	UPROPERTY(meta = (Input))
 	FRigComponentKey PhysicsJointComponentKey;
 
+	// The current Physics Joint Drive data
 	UPROPERTY(meta = (Output))
 	FRigPhysicsDriveData DriveData;
 };
@@ -186,9 +196,12 @@ struct FRigUnit_HierarchySetJointDriveUseSkeletalAnimation : public FRigUnit_Phy
 	RIGVM_METHOD()
 	UE_API virtual void Execute() override;
 
+	// The Physics Joint to be updated
 	UPROPERTY(meta = (Input))
 	FRigComponentKey PhysicsJointComponentKey;
 
+	// Whether the drive targets should be relative to the incoming animation pose. If set to true
+	// (and the drive targets are zero), then the drives will track the target animation.
 	UPROPERTY(meta = (Input))
 	bool bUseSkeletalAnimation = true;
 };
@@ -203,20 +216,23 @@ struct FRigUnit_MakeArticulationJointData : public FRigUnit_PhysicsBase
 	RIGVM_METHOD()
 	UE_API virtual void Execute() override;
 
-	// Degrees, twist, swing1, swing2
+	// The angular limit in Degrees (twist, swing1, swing2)
 	// -ve indicates the limit range is free
 	UPROPERTY(meta = (Input, ClampMin = "-1.0"))
 	FVector AngularLimit = FVector(-1, -1, -1);
 
-	// If limited, then this will be used to control the softness
+	// If limited, then this will be used to control the softness of the limit
 	// -ve indicates the limit is hard
 	// A value of 1 is reasonably soft
 	UPROPERTY(meta = (Input, ClampMin = "-1.0"))
 	FVector SoftStrength = FVector(-1, -1, -1);
 
+	// If limited, then this will be used to control the damping ratio of the limit
 	UPROPERTY(meta = (Input, ClampMin = "0.0"))
 	FVector SoftDampingRatio = FVector(1, 1, 1);
 
+	// Joint data that can be used in a Physics Joint, calculated according to the values set here,
+	// and configured to act as a "normal" articulation joint.
 	UPROPERTY(meta = (Output))
 	FRigPhysicsJointData JointData;
 };
@@ -239,22 +255,19 @@ struct FRigUnit_MakeArticulationDriveData : public FRigUnit_PhysicsBase
 	UPROPERTY(meta = (Input))
 	TEnumAsByte<enum EAngularDriveMode::Type> AngularDriveMode = EAngularDriveMode::SLERP;
 
-	/** The strength used to drive angular motion */
+	// The strength used to drive angular motion
 	UPROPERTY(meta = (Input, ClampMin = "0.0"))
 	float AngularStrength = 10.0f;
 
-	/**
-	 * The amount of damping associated with the angular strength. A value of 1 Results in critically
-	 * damped motion where the control drives as quickly as possible to the target without overshooting.
-	 * Values > 1 result in more damped motion, and values below 1 result in faster, but more "wobbly" motion.
-	 */
+	// The amount of damping associated with the angular strength. A value of 1 Results in
+	// critically damped motion where the control drives as quickly as possible to the target
+	// without overshooting. Values > 1 result in more damped motion, and values below 1 result in
+	// faster, but more "wobbly" motion.
 	UPROPERTY(meta = (Input, ClampMin = "0.0"))
 	float AngularDampingRatio = 1.0f;
 
-	/**
-	 * The amount of additional angular damping. This is added to the damping that comes from AngularDampingRatio
-	 * and can be useful when you want damping even when AngularStrength is zero.
-	 */
+	// The amount of additional angular damping. This is added to the damping that comes from
+	// AngularDampingRatio and can be useful when you want damping even when AngularStrength is zero.
 	UPROPERTY(meta = (Input, ClampMin = "0.0"))
 	float AngularExtraDamping = 0.0f;
 
@@ -262,11 +275,11 @@ struct FRigUnit_MakeArticulationDriveData : public FRigUnit_PhysicsBase
 	UPROPERTY(meta = (Input, ClampMin = "0.0"))
 	float SkeletalAnimationVelocityMultiplier = 1.0f;
 
+	// The drive data that can be set in a Physics Joint, reflecting the setup specified here.
 	UPROPERTY(meta = (Output))
 	FRigPhysicsDriveData DriveData;
 };
 
-
 // Helper to simplify creation of drive data
 USTRUCT(meta = (DisplayName = "Make Drive Data", Keywords = "Construction,Create,New,Body,Skeleton", Varying))
 struct FRigUnit_MakeDriveData : public FRigUnit_PhysicsBase
@@ -280,49 +293,42 @@ struct FRigUnit_MakeDriveData : public FRigUnit_PhysicsBase
 	UPROPERTY(meta = (Input))
 	bool bEnableLinearDrive = false;
 
-	/** The strength used to drive linear motion */
+	// The strength used to drive linear motion
 	UPROPERTY(meta = (Input, ClampMin = "0.0"))
 	float LinearStrength = 10.0f;
 
-	/**
-	 * The amount of damping associated with the linear strength. A value of 1 Results in critically
-	 * damped motion where the control drives as quickly as possible to the target without overshooting.
-	 * Values > 1 result in more damped motion, and values below 1 result in faster, but more "wobbly" motion.
-	 */
+	// The amount of damping associated with the linear strength. A value of 1 Results in critically
+	// damped motion where the control drives as quickly as possible to the target without
+	// overshooting. Values > 1 result in more damped motion, and values below 1 result in faster,
+	// but more "wobbly" motion.
 	UPROPERTY(meta = (Input, ClampMin = "0.0"))
 	float LinearDampingRatio = 1.0f;
 
-	/**
-	 * The amount of additional linear damping. This is added to the damping that comes from LinearDampingRatio
-	 * and can be useful when you want damping even when LinearStrength is zero.
-	 */
+	// The amount of additional linear damping. This is added to the damping that comes from
+	// LinearDampingRatio and can be useful when you want damping even when LinearStrength is zero.
 	UPROPERTY(meta = (Input, ClampMin = "0.0"))
 	float LinearExtraDamping = 0.0f;
 
-	UPROPERTY(meta = (Input))
 	// Whether to enable the angular drive
+	UPROPERTY(meta = (Input))
 	bool bEnableAngularDrive = true;
 
 	// The type of drive. Note that SLERP drives don't work if any axis is locked
 	UPROPERTY(meta = (Input))
 	TEnumAsByte<enum EAngularDriveMode::Type> AngularDriveMode = EAngularDriveMode::SLERP;
 
-	/** The strength used to drive angular motion */
+	// The strength used to drive angular motion 
 	UPROPERTY(meta = (Input, ClampMin = "0.0"))
 	float AngularStrength = 10.0f;
 
-	/**
-	 * The amount of damping associated with the angular strength. A value of 1 Results in critically
-	 * damped motion where the control drives as quickly as possible to the target without overshooting.
-	 * Values > 1 result in more damped motion, and values below 1 result in faster, but more "wobbly" motion.
-	 */
+	// The amount of damping associated with the angular strength. A value of 1 Results in critically
+	// damped motion where the control drives as quickly as possible to the target without overshooting.
+	// Values > 1 result in more damped motion, and values below 1 result in faster, but more "wobbly" motion.
 	UPROPERTY(meta = (Input, ClampMin = "0.0"))
 	float AngularDampingRatio = 1.0f;
 
-	/**
-	 * The amount of additional angular damping. This is added to the damping that comes from AngularDampingRatio
-	 * and can be useful when you want damping even when AngularStrength is zero.
-	 */
+	// The amount of additional angular damping. This is added to the damping that comes from AngularDampingRatio
+	// and can be useful when you want damping even when AngularStrength is zero.
 	UPROPERTY(meta = (Input, ClampMin = "0.0"))
 	float AngularExtraDamping = 0.0f;
 
@@ -330,6 +336,7 @@ struct FRigUnit_MakeDriveData : public FRigUnit_PhysicsBase
 	UPROPERTY(meta = (Input, ClampMin = "0.0"))
 	float SkeletalAnimationVelocityMultiplier = 1.0f;
 
+	// The drive data, representing the properties specified here, that can be set on a Physics Joint
 	UPROPERTY(meta = (Output))
 	FRigPhysicsDriveData DriveData;
 };
