Fix minor typo in comment

#rnx

diff --git a/Engine/Plugins/Experimental/ControlRigPhysics/Source/ControlRigPhysics/Public/RigPhysicsJointExecution.h b/Engine/Plugins/Experimental/ControlRigPhysics/Source/ControlRigPhysics/Public/RigPhysicsJointExecution.h
--- a/Engine/Plugins/Experimental/ControlRigPhysics/Source/ControlRigPhysics/Public/RigPhysicsJointExecution.h
+++ b/Engine/Plugins/Experimental/ControlRigPhysics/Source/ControlRigPhysics/Public/RigPhysicsJointExecution.h
@@ -15,7 +15,7 @@
 #define UE_API CONTROLRIGPHYSICS_API
 
 /**
- * Adds a new physics body as a component on the owner element.
+ * Adds a new Physics Joint as a component on the owner element.
  * Note: This node only runs as part of the construction event.
  */
 USTRUCT(meta = (DisplayName = "Spawn Physics Joint", Keywords = "Add,Construction,Create,New,Joint,Articulation,Constraint,Skeleton", Varying))
