Control Rig Physics - Expose more properties to the graph

diff --git a/Engine/Plugins/Experimental/ControlRigPhysics/Source/ControlRigPhysics/Private/RigPhysicsControlExecution.cpp b/Engine/Plugins/Experimental/ControlRigPhysics/Source/ControlRigPhysics/Private/RigPhysicsControlExecution.cpp
--- a/Engine/Plugins/Experimental/ControlRigPhysics/Source/ControlRigPhysics/Private/RigPhysicsControlExecution.cpp
+++ b/Engine/Plugins/Experimental/ControlRigPhysics/Source/ControlRigPhysics/Private/RigPhysicsControlExecution.cpp
@@ -120,6 +120,21 @@ FRigUnit_HierarchySetControlLinearDampingRatio_Execute()
 	}
 }
 
+//======================================================================================================================
+FRigUnit_HierarchySetControlTargetVelocityMultipliers_Execute()
+{
+	DECLARE_SCOPE_HIERARCHICAL_COUNTER_RIGUNIT();
+	if (ExecuteContext.Hierarchy)
+	{
+		if (FRigPhysicsControlComponent* Component = Cast<FRigPhysicsControlComponent>(
+			ExecuteContext.Hierarchy->FindComponent(PhysicsControlComponentKey)))
+		{
+			Component->ControlData.LinearTargetVelocityMultiplier = LinearTargetVelocityMultiplier;
+			Component->ControlData.AngularTargetVelocityMultiplier = AngularTargetVelocityMultiplier;
+		}
+	}
+}
+
 //======================================================================================================================
 FRigUnit_HierarchySetControlAngularStrength_Execute()
 {
diff --git a/Engine/Plugins/Experimental/ControlRigPhysics/Source/ControlRigPhysics/Private/RigPhysicsJointExecution.cpp b/Engine/Plugins/Experimental/ControlRigPhysics/Source/ControlRigPhysics/Private/RigPhysicsJointExecution.cpp
--- a/Engine/Plugins/Experimental/ControlRigPhysics/Source/ControlRigPhysics/Private/RigPhysicsJointExecution.cpp
+++ b/Engine/Plugins/Experimental/ControlRigPhysics/Source/ControlRigPhysics/Private/RigPhysicsJointExecution.cpp
@@ -82,6 +82,51 @@ FRigUnit_HierarchyGetJointData_Execute()
 	}
 }
 
+//======================================================================================================================
+FRigUnit_HierarchySetJointLinearLimit_Execute()
+{
+	DECLARE_SCOPE_HIERARCHICAL_COUNTER_RIGUNIT();
+	if (!ExecuteContext.Hierarchy)
+	{
+		return;
+	}
+	if (FRigPhysicsJointComponent* Component = Cast<FRigPhysicsJointComponent>(
+		ExecuteContext.Hierarchy->FindComponent(PhysicsJointComponentKey)))
+	{
+		Component->JointData.LinearConstraint = LinearLimit.LinearConstraint;
+	}
+}
+
+//======================================================================================================================
+FRigUnit_HierarchySetJointConeLimit_Execute()
+{
+	DECLARE_SCOPE_HIERARCHICAL_COUNTER_RIGUNIT();
+	if (!ExecuteContext.Hierarchy)
+	{
+		return;
+	}
+	if (FRigPhysicsJointComponent* Component = Cast<FRigPhysicsJointComponent>(
+		ExecuteContext.Hierarchy->FindComponent(PhysicsJointComponentKey)))
+	{
+		Component->JointData.ConeConstraint = ConeLimit.ConeConstraint;
+	}
+}
+
+//======================================================================================================================
+FRigUnit_HierarchySetJointTwistLimit_Execute()
+{
+	DECLARE_SCOPE_HIERARCHICAL_COUNTER_RIGUNIT();
+	if (!ExecuteContext.Hierarchy)
+	{
+		return;
+	}
+	if (FRigPhysicsJointComponent* Component = Cast<FRigPhysicsJointComponent>(
+		ExecuteContext.Hierarchy->FindComponent(PhysicsJointComponentKey)))
+	{
+		Component->JointData.TwistConstraint = TwistLimit.TwistConstraint;
+	}
+}
+
 //======================================================================================================================
 FRigUnit_HierarchySetJointEnabled_Execute()
 {
diff --git a/Engine/Plugins/Experimental/ControlRigPhysics/Source/ControlRigPhysics/Private/RigPhysicsSimulation.cpp b/Engine/Plugins/Experimental/ControlRigPhysics/Source/ControlRigPhysics/Private/RigPhysicsSimulation.cpp
--- a/Engine/Plugins/Experimental/ControlRigPhysics/Source/ControlRigPhysics/Private/RigPhysicsSimulation.cpp
+++ b/Engine/Plugins/Experimental/ControlRigPhysics/Source/ControlRigPhysics/Private/RigPhysicsSimulation.cpp
@@ -1088,6 +1088,11 @@ void FRigPhysicsSimulation::UpdateJointPrePhysics(
 		Constraint->SetConstraintEnabled(JointData.bEnabled);
 
 		Chaos::FPBDJointSettings Settings = Constraint->GetSettings();
+
+		ImmediatePhysics::UpdateJointSettingsFromLinearConstraint(JointData.LinearConstraint, Settings);
+		ImmediatePhysics::UpdateJointSettingsFromConeConstraint(JointData.ConeConstraint, Settings);
+		ImmediatePhysics::UpdateJointSettingsFromTwistConstraint(JointData.TwistConstraint, Settings);
+
 		ImmediatePhysics::UpdateJointSettingsFromLinearDriveConstraint(DriveData.LinearDriveConstraint, Settings);
 		ImmediatePhysics::UpdateJointSettingsFromAngularDriveConstraint(DriveData.AngularDriveConstraint, Settings);
 
diff --git a/Engine/Plugins/Experimental/ControlRigPhysics/Source/ControlRigPhysics/Public/RigPhysicsControlExecution.h b/Engine/Plugins/Experimental/ControlRigPhysics/Source/ControlRigPhysics/Public/RigPhysicsControlExecution.h
--- a/Engine/Plugins/Experimental/ControlRigPhysics/Source/ControlRigPhysics/Public/RigPhysicsControlExecution.h
+++ b/Engine/Plugins/Experimental/ControlRigPhysics/Source/ControlRigPhysics/Public/RigPhysicsControlExecution.h
@@ -254,6 +254,36 @@ struct FRigUnit_HierarchySetControlAngularDampingRatio : public FRigUnit_Physics
 	float DampingRatio = 1.0f;
 };
 
+// Sets the Target velocity multipliers of a Physics Control
+USTRUCT(meta = (DisplayName = "Set Physics Control Target Velocity Multipliers", Varying))
+struct FRigUnit_HierarchySetControlTargetVelocityMultipliers : public FRigUnit_PhysicsBaseMutable
+{
+	GENERATED_BODY()
+
+	FRigUnit_HierarchySetControlTargetVelocityMultipliers()
+	{
+		PhysicsControlComponentKey.ElementKey.Type = ERigElementType::Bone;
+		PhysicsControlComponentKey.Name = FRigPhysicsControlComponent::GetDefaultName();
+	}
+
+	RIGVM_METHOD()
+	UE_API virtual void Execute() override;
+
+	// The Physics Control component to be updated
+	UPROPERTY(meta = (Input))
+	FRigComponentKey PhysicsControlComponentKey;
+
+	// The multiplier for the linear target velocity. Set this to zero to see sim-space damping. Set
+	// it to one to track the targets more closely.
+	UPROPERTY(meta = (Input))
+	float LinearTargetVelocityMultiplier = 1.0f;
+
+	// The multiplier for the angular target velocity. Set this to zero to see sim-space damping. Set
+	// it to one to track the targets more closely.
+	UPROPERTY(meta = (Input))
+	float AngularTargetVelocityMultiplier = 1.0f;
+};
+
 // Gets the control data for a physics control
 USTRUCT(meta = (DisplayName = "Get Physics Control Data", Varying))
 struct FRigUnit_HierarchyGetControlData : public FRigUnit_PhysicsBase
diff --git a/Engine/Plugins/Experimental/ControlRigPhysics/Source/ControlRigPhysics/Public/RigPhysicsData.h b/Engine/Plugins/Experimental/ControlRigPhysics/Source/ControlRigPhysics/Public/RigPhysicsData.h
--- a/Engine/Plugins/Experimental/ControlRigPhysics/Source/ControlRigPhysics/Public/RigPhysicsData.h
+++ b/Engine/Plugins/Experimental/ControlRigPhysics/Source/ControlRigPhysics/Public/RigPhysicsData.h
@@ -22,6 +22,10 @@ class FRigElementKeyRedirector;
 
 const float RigPhysicsDefaultShapeSize = 10.0f;
 
+UE_API FArchive& operator <<(FArchive& Ar, FConstraintDrive& Data);
+UE_API FArchive& operator <<(FArchive& Ar, FLinearDriveConstraint& Data);
+UE_API FArchive& operator <<(FArchive& Ar, FAngularDriveConstraint& Data);
+
 USTRUCT(BlueprintType)
 struct FRigPhysicsVisualizationSettings
 {
@@ -634,6 +638,33 @@ struct FRigPhysicsDriveData
 	float SkeletalAnimationVelocityMultiplier = 1.0f;
 };
 
+USTRUCT(BlueprintType)
+struct FRigPhysicsLinearLimit
+{
+	GENERATED_BODY()
+
+	UPROPERTY(EditAnywhere, Category = Limit)
+	FLinearConstraint LinearConstraint;
+};
+
+USTRUCT(BlueprintType)
+struct FRigPhysicsConeLimit
+{
+	GENERATED_BODY()
+
+	UPROPERTY(EditAnywhere, Category = Limit)
+	FConeConstraint ConeConstraint;
+};
+
+USTRUCT(BlueprintType)
+struct FRigPhysicsTwistLimit
+{
+	GENERATED_BODY()
+
+	UPROPERTY(EditAnywhere, Category = Limit)
+	FTwistConstraint TwistConstraint;
+};
+
 /**
  * This represents an "articulation" - a general purpose "character joint" consisting of
  *   * A linear limit (a linear constraint), which is normally used to pin two bones together
diff --git a/Engine/Plugins/Experimental/ControlRigPhysics/Source/ControlRigPhysics/Public/RigPhysicsJointExecution.h b/Engine/Plugins/Experimental/ControlRigPhysics/Source/ControlRigPhysics/Public/RigPhysicsJointExecution.h
--- a/Engine/Plugins/Experimental/ControlRigPhysics/Source/ControlRigPhysics/Public/RigPhysicsJointExecution.h
+++ b/Engine/Plugins/Experimental/ControlRigPhysics/Source/ControlRigPhysics/Public/RigPhysicsJointExecution.h
@@ -61,7 +61,7 @@ struct FRigUnit_AddPhysicsJoint : public FRigUnit_PhysicsBaseMutable
 	FRigPhysicsDriveData DriveData;
 };
 
-// Sets the joint for a physics component body
+// Sets the joint data for a physics joint component
 USTRUCT(meta = (DisplayName = "Set Physics Joint Properties", Varying))
 struct FRigUnit_HierarchySetJointData : public FRigUnit_PhysicsBaseMutable
 {
@@ -85,7 +85,7 @@ struct FRigUnit_HierarchySetJointData : public FRigUnit_PhysicsBaseMutable
 	FRigPhysicsJointData JointData;
 };
 
-// Gets the joint for a physics joint component
+// Gets the joint data for a physics joint component
 USTRUCT(meta = (DisplayName = "Get Physics Joint Properties"))
 struct FRigUnit_HierarchyGetJointData : public FRigUnit_PhysicsBase
 {
@@ -109,6 +109,78 @@ struct FRigUnit_HierarchyGetJointData : public FRigUnit_PhysicsBase
 	FRigPhysicsJointData JointData;
 };
 
+// Sets the joint linear limit
+USTRUCT(meta = (DisplayName = "Set Physics Joint Linear Constraint Limit", Varying))
+struct FRigUnit_HierarchySetJointLinearLimit : public FRigUnit_PhysicsBaseMutable
+{
+	GENERATED_BODY()
+
+	FRigUnit_HierarchySetJointLinearLimit()
+	{
+		PhysicsJointComponentKey.ElementKey.Type = ERigElementType::Bone;
+		PhysicsJointComponentKey.Name = FRigPhysicsJointComponent::GetDefaultName();
+	}
+
+	RIGVM_METHOD()
+	UE_API virtual void Execute() override;
+
+	// The Physics Joint to be updated
+	UPROPERTY(meta = (Input))
+	FRigComponentKey PhysicsJointComponentKey;
+
+	// The linear constraint to be used
+	UPROPERTY(meta = (Input))
+	FRigPhysicsLinearLimit LinearLimit;
+};
+
+// Sets the joint cone limit
+USTRUCT(meta = (DisplayName = "Set Physics Joint Cone Constraint Limit", Varying))
+struct FRigUnit_HierarchySetJointConeLimit : public FRigUnit_PhysicsBaseMutable
+{
+	GENERATED_BODY()
+
+	FRigUnit_HierarchySetJointConeLimit()
+	{
+		PhysicsJointComponentKey.ElementKey.Type = ERigElementType::Bone;
+		PhysicsJointComponentKey.Name = FRigPhysicsJointComponent::GetDefaultName();
+	}
+
+	RIGVM_METHOD()
+	UE_API virtual void Execute() override;
+
+	// The Physics Joint to be updated
+	UPROPERTY(meta = (Input))
+	FRigComponentKey PhysicsJointComponentKey;
+
+	// The cone limit
+	UPROPERTY(meta = (Input))
+	FRigPhysicsConeLimit ConeLimit;
+};
+
+// Sets the joint twist limit
+USTRUCT(meta = (DisplayName = "Set Physics Joint Twist Constraint Limit", Varying))
+struct FRigUnit_HierarchySetJointTwistLimit : public FRigUnit_PhysicsBaseMutable
+{
+	GENERATED_BODY()
+
+	FRigUnit_HierarchySetJointTwistLimit()
+	{
+		PhysicsJointComponentKey.ElementKey.Type = ERigElementType::Bone;
+		PhysicsJointComponentKey.Name = FRigPhysicsJointComponent::GetDefaultName();
+	}
+
+	RIGVM_METHOD()
+	UE_API virtual void Execute() override;
+
+	// The Physics Joint to be updated
+	UPROPERTY(meta = (Input))
+	FRigComponentKey PhysicsJointComponentKey;
+
+	// The twist limit
+	UPROPERTY(meta = (Input))
+	FRigPhysicsTwistLimit TwistLimit;
+};
+
 // Specifies whether a Physics Joint should be enabled or not
 USTRUCT(meta = (DisplayName = "Set Physics Joint Enabled", Varying))
 struct FRigUnit_HierarchySetJointEnabled : public FRigUnit_PhysicsBaseMutable
