Control Rig Physics: Remove some TODOs from public-facing comments, and some associated tidy up of comments and minor functionality.

diff --git a/Engine/Plugins/Animation/ControlRig/Source/ControlRig/Private/ControlRig.cpp b/Engine/Plugins/Animation/ControlRig/Source/ControlRig/Private/ControlRig.cpp
--- a/Engine/Plugins/Animation/ControlRig/Source/ControlRig/Private/ControlRig.cpp
+++ b/Engine/Plugins/Animation/ControlRig/Source/ControlRig/Private/ControlRig.cpp
@@ -731,9 +731,8 @@ void UControlRig::InitializeFromCDO()
 			UControlRig* CDO = GetClass()->GetDefaultObject<UControlRig>();
 			URigHierarchy* Hierarchy = GetHierarchy();
 
-			// Danny TODO this could copy the physics solvers from the CDO, but they are transient
-			// objects that only exist through construction, so I think we just need to start afresh
-			// each time. 
+			// Physics Simulations are transient objects that only exist through construction so do
+			// not need to be initialized from the CDO
 			PhysicsSimulations.Empty();
 
 			// copy hierarchy
diff --git a/Engine/Plugins/Experimental/ControlRigPhysics/Source/ControlRigPhysics/Private/RigPhysicsBodyExecution.cpp b/Engine/Plugins/Experimental/ControlRigPhysics/Source/ControlRigPhysics/Private/RigPhysicsBodyExecution.cpp
--- a/Engine/Plugins/Experimental/ControlRigPhysics/Source/ControlRigPhysics/Private/RigPhysicsBodyExecution.cpp
+++ b/Engine/Plugins/Experimental/ControlRigPhysics/Source/ControlRigPhysics/Private/RigPhysicsBodyExecution.cpp
@@ -178,6 +178,9 @@ FRigUnit_HierarchySetPhysicsBodyKinematicTargetSpace_Execute()
 }
 
 //======================================================================================================================
+// Danny TODO Note that the sparse data does not get displayed correctly, so this is largely
+// unusable - the flags that enable / disable all end up getting reset if the user attempts to
+// change them.
 FRigUnit_HierarchySetPhysicsBodySparseData_Execute()
 {
 	DECLARE_SCOPE_HIERARCHICAL_COUNTER_RIGUNIT();
@@ -274,6 +277,8 @@ FRigUnit_HierarchySetPhysicsBodyGravityMultiplier_Execute()
 }
 
 //======================================================================================================================
+// Danny TODO implement fully. Note that this isn't completely obvious - sometimes we want to blend 
+// back in the hierarchy space, and sometimes in simulation space.
 FRigUnit_HierarchySetPhysicsBodyPhysicsBlendWeight_Execute()
 {
 	DECLARE_SCOPE_HIERARCHICAL_COUNTER_RIGUNIT();
diff --git a/Engine/Plugins/Experimental/ControlRigPhysics/Source/ControlRigPhysics/Public/RigPhysicsBodyExecution.h b/Engine/Plugins/Experimental/ControlRigPhysics/Source/ControlRigPhysics/Public/RigPhysicsBodyExecution.h
--- a/Engine/Plugins/Experimental/ControlRigPhysics/Source/ControlRigPhysics/Public/RigPhysicsBodyExecution.h
+++ b/Engine/Plugins/Experimental/ControlRigPhysics/Source/ControlRigPhysics/Public/RigPhysicsBodyExecution.h
@@ -221,8 +221,8 @@ struct FRigUnit_HierarchySetPhysicsBodyTargetBone : public FRigUnit_PhysicsBaseM
 
 /**
  * Sets all the data on a body - but in a sparse way so you can decide which parameters get applied.
- * Danny TODO - Note that the sparse data does not get displayed correctly, so this is largely 
- * unusable - the flags that enable/disable all end up getting reset if the user attempts to change them.
+ * Note that the sparse data does not get displayed correctly (the flags that enable/disable all end up getting 
+ * reset if the user attempts to change them) so this should be used with caution until this is fixed.
  */
 USTRUCT(meta = (DisplayName = "Set Physics Body Data", Varying))
 struct FRigUnit_HierarchySetPhysicsBodySparseData : public FRigUnit_PhysicsBaseMutable
@@ -428,8 +428,7 @@ struct FRigUnit_HierarchySetPhysicsBodyGravityMultiplier : public FRigUnit_Physi
 
 /**
  * Controls the amount that the simulation is blended back into the target bones. 
- * Danny TODO implement. Note that this isn't completely obvious - sometimes we want to blend 
- * back in the hierarchy space, and sometimes in simulation space.
+ * Note that this feature is not yet fully implemented.
  */
 USTRUCT(meta = (DisplayName = "Set Physics Body Physics Blend Weight", Varying))
 struct FRigUnit_HierarchySetPhysicsBodyPhysicsBlendWeight : public FRigUnit_PhysicsBaseMutable
@@ -457,7 +456,6 @@ struct FRigUnit_HierarchySetPhysicsBodyPhysicsBlendWeight : public FRigUnit_Phys
 /**
  * If true, then kinematic objects will be written back from simulation to the bones. This only 
  * necessary when either kinematic targets are being used, or when the target bone differs from the source bone.
- * Danny TODO - check the implementation, and also can we auto-calculate this at runtime?
  */
 USTRUCT(meta = (DisplayName = "Set Physics Body Update Kinematic From Simulation", Varying))
 struct FRigUnit_HierarchySetPhysicsBodyUpdateKinematicFromSimulation : public FRigUnit_PhysicsBaseMutable
diff --git a/Engine/Plugins/Experimental/ControlRigPhysics/Source/ControlRigPhysics/Public/RigPhysicsControlExecution.h b/Engine/Plugins/Experimental/ControlRigPhysics/Source/ControlRigPhysics/Public/RigPhysicsControlExecution.h
--- a/Engine/Plugins/Experimental/ControlRigPhysics/Source/ControlRigPhysics/Public/RigPhysicsControlExecution.h
+++ b/Engine/Plugins/Experimental/ControlRigPhysics/Source/ControlRigPhysics/Public/RigPhysicsControlExecution.h
@@ -168,8 +168,9 @@ struct FRigUnit_HierarchySetControlLinearStrength : public FRigUnit_PhysicsBaseM
 
 	// The linear strength of the control. In the absence of any other influences, this is the
 	// frequency of oscillation when there is zero damping. For example, a strength of 2 would make
-	// the target oscillate twice per second (when there is no damping).
-	// Note that parent-space controls should normally have this set to zero, since 
+	// the target oscillate twice per second (when there is no damping). Note that parent-space
+	// controls should this set to zero when they are controlling two bodies that are connected with
+	// a physics joint.
 	UPROPERTY(meta = (Input, ClampMin = "0.0"))
 	float Strength = 0.0f;
 };
diff --git a/Engine/Plugins/Experimental/PhysicsControl/Source/PhysicsControl/Private/RigidBodyWithControl.cpp b/Engine/Plugins/Experimental/PhysicsControl/Source/PhysicsControl/Private/RigidBodyWithControl.cpp
--- a/Engine/Plugins/Experimental/PhysicsControl/Source/PhysicsControl/Private/RigidBodyWithControl.cpp
+++ b/Engine/Plugins/Experimental/PhysicsControl/Source/PhysicsControl/Private/RigidBodyWithControl.cpp
@@ -170,6 +170,8 @@ void UpdateBodyFromModifierData(
 	// Set collision
 	PhysicsSimulation->SetHasCollision(
 		ActorHandle, CollisionEnabledHasPhysics(ModifierData.CollisionType));
+
+	ActorHandle->SetUpdateKinematicFromSimulation(ModifierData.bUpdateKinematicFromSimulation);
 }
 
 //======================================================================================================================
diff --git a/Engine/Plugins/Experimental/PhysicsControl/Source/PhysicsControl/Public/PhysicsControlData.h b/Engine/Plugins/Experimental/PhysicsControl/Source/PhysicsControl/Public/PhysicsControlData.h
--- a/Engine/Plugins/Experimental/PhysicsControl/Source/PhysicsControl/Public/PhysicsControlData.h
+++ b/Engine/Plugins/Experimental/PhysicsControl/Source/PhysicsControl/Public/PhysicsControlData.h
@@ -295,7 +295,6 @@ struct FPhysicsControlModifierData
 	 * If true then the associated actor's transform will be updated from the simulation when it is
 	 * kinematic. This is most likely useful when using async physics in order to prevent different
 	 * parts of the skeleton from being torn apart.
-	 * Note - this is not used in RBWC
 	 */
 	UPROPERTY(EditAnywhere, BlueprintReadWrite, Category = PhysicsControl)
 	bool bUpdateKinematicFromSimulation;
diff --git a/Engine/Source/Runtime/Engine/Private/PhysicsEngine/ImmediatePhysics/ImmediatePhysicsChaos/ImmediatePhysicsActorHandle_Chaos.cpp b/Engine/Source/Runtime/Engine/Private/PhysicsEngine/ImmediatePhysics/ImmediatePhysicsChaos/ImmediatePhysicsActorHandle_Chaos.cpp
--- a/Engine/Source/Runtime/Engine/Private/PhysicsEngine/ImmediatePhysics/ImmediatePhysicsChaos/ImmediatePhysicsActorHandle_Chaos.cpp
+++ b/Engine/Source/Runtime/Engine/Private/PhysicsEngine/ImmediatePhysics/ImmediatePhysicsChaos/ImmediatePhysicsActorHandle_Chaos.cpp
@@ -321,6 +321,25 @@ namespace ImmediatePhysics_Chaos
 		return ParticleHandle->CastToRigidParticle() != nullptr;
 	}
 
+	void FActorHandle::SetUpdateKinematicFromSimulation(bool bUpdateKinematicFromSimulation)
+	{
+		using namespace Chaos;
+		if (FPBDRigidParticleHandle* Rigid = Handle()->CastToRigidParticle())
+		{
+			Rigid->SetUpdateKinematicFromSimulation(bUpdateKinematicFromSimulation);
+		}
+	}
+
+	bool FActorHandle::GetUpdateKinematicFromSimulation() const
+	{
+		using namespace Chaos;
+		if (const FPBDRigidParticleHandle* Rigid = Handle()->CastToRigidParticle())
+		{
+			return Rigid->UpdateKinematicFromSimulation();
+		}
+		return false;
+	}
+
 	bool FActorHandle::IsGravityEnabled() const
 	{
 		using namespace Chaos;
diff --git a/Engine/Source/Runtime/Engine/Public/Physics/ImmediatePhysics/ImmediatePhysicsChaos/ImmediatePhysicsActorHandle_Chaos.h b/Engine/Source/Runtime/Engine/Public/Physics/ImmediatePhysics/ImmediatePhysicsChaos/ImmediatePhysicsActorHandle_Chaos.h
--- a/Engine/Source/Runtime/Engine/Public/Physics/ImmediatePhysics/ImmediatePhysicsChaos/ImmediatePhysicsActorHandle_Chaos.h
+++ b/Engine/Source/Runtime/Engine/Public/Physics/ImmediatePhysics/ImmediatePhysicsChaos/ImmediatePhysicsActorHandle_Chaos.h
@@ -98,6 +98,12 @@ namespace ImmediatePhysics_Chaos
 		/** Returns true if the world gravity applies to this actor */
 		ENGINE_API bool IsGravityEnabled() const;
 
+		/** Returns true if the output transform is set to the final body transform when it is kinematic. Will return false if the body is not a rigid. */
+		ENGINE_API bool GetUpdateKinematicFromSimulation() const;
+
+		/** Sets whether or not the output transform is set to the final body transform when it is kinematic. */
+		ENGINE_API void SetUpdateKinematicFromSimulation(bool bUpdateKinematicFromSimulation);
+
 		/** Returns true the body is or could be dynamic (i.e. currently simulated, or a rigid body that is currently kinematic) */
 		ENGINE_API bool CouldBeDynamic() const;
 
