Control Rig Physics - collision fixes for when bodies have no shapes

diff --git a/Engine/Plugins/Experimental/ControlRigPhysics/Source/ControlRigPhysics/Private/RigPhysicsSimulationWorldObjects.cpp b/Engine/Plugins/Experimental/ControlRigPhysics/Source/ControlRigPhysics/Private/RigPhysicsSimulationWorldObjects.cpp
--- a/Engine/Plugins/Experimental/ControlRigPhysics/Source/ControlRigPhysics/Private/RigPhysicsSimulationWorldObjects.cpp
+++ b/Engine/Plugins/Experimental/ControlRigPhysics/Source/ControlRigPhysics/Private/RigPhysicsSimulationWorldObjects.cpp
@@ -109,9 +109,12 @@ void FRigPhysicsSimulation::UpdateWorldObjectsPostPhysics(
 				const FTransform ParticleTransform = GeometryParticleHandle->GetTransformXR() * SpaceTransform;
 
 				const auto& AABB = GeometryParticleHandle->LocalBounds();
-				for (int32 Index = 0 ; Index != 8 ; ++Index)
+				if (!AABB.IsEmpty())
 				{
-					WorldOverlapBox += ParticleTransform.TransformPosition(AABB.GetVertex(Index));
+					for (int32 Index = 0; Index != 8; ++Index)
+					{
+						WorldOverlapBox += ParticleTransform.TransformPosition(AABB.GetVertex(Index));
+					}
 				}
 			}
 		}
