Support for applying forces in the control rig physics and RigidBodyWithControl animation graph nodes

#Synced-CL 48547444

diff --git a/Engine/Plugins/Experimental/ControlRigPhysics/Source/ControlRigPhysics/Private/RigPhysicsBodyExecution.cpp b/Engine/Plugins/Experimental/ControlRigPhysics/Source/ControlRigPhysics/Private/RigPhysicsBodyExecution.cpp
--- a/Engine/Plugins/Experimental/ControlRigPhysics/Source/ControlRigPhysics/Private/RigPhysicsBodyExecution.cpp
+++ b/Engine/Plugins/Experimental/ControlRigPhysics/Source/ControlRigPhysics/Private/RigPhysicsBodyExecution.cpp
@@ -15,6 +15,12 @@
 
 #include UE_INLINE_GENERATED_CPP_BY_NAME(RigPhysicsBodyExecution)
 
+//======================================================================================================================
+inline FRigPhysicsBodyComponent* GetBody(URigHierarchy* Hierarchy, const FRigComponentKey& Key)
+{
+	return Cast<FRigPhysicsBodyComponent>(Hierarchy->FindComponent(Key));
+}
+
 //======================================================================================================================
 FRigUnit_AddPhysicsBody_Execute()
 {
@@ -35,8 +41,7 @@ FRigUnit_AddPhysicsBody_Execute()
 			FRigPhysicsBodyComponent::StaticStruct(), FRigPhysicsBodyComponent::GetDefaultName(), Owner);
 		if (PhysicsBodyComponentKey.IsValid())
 		{
-			if (FRigPhysicsBodyComponent* Component = Cast<FRigPhysicsBodyComponent>(
-				ExecuteContext.Hierarchy->FindComponent(PhysicsBodyComponentKey)))
+			if (FRigPhysicsBodyComponent* Component = GetBody(ExecuteContext.Hierarchy, PhysicsBodyComponentKey))
 			{
 				Component->BodySolverSettings = Solver;
 				Component->Dynamics = Dynamics;
@@ -64,8 +69,7 @@ FRigUnit_HierarchyAutoCalculateCollision_Execute()
 		UE_CONTROLRIG_RIGUNIT_REPORT_ERROR(TEXT("AutoCalculateCollision can only be used during Setup"));
 	}
 
-	if (FRigPhysicsBodyComponent* Component = Cast<FRigPhysicsBodyComponent>(
-		ExecuteContext.Hierarchy->FindComponent(PhysicsBodyComponentKey)))
+	if (FRigPhysicsBodyComponent* Component = GetBody(ExecuteContext.Hierarchy, PhysicsBodyComponentKey))
 	{
 		Component->AutoCalculateCollision(ExecuteContext.Hierarchy, MinAspectRatio, MinSize);
 	}
@@ -79,8 +83,7 @@ FRigUnit_HierarchySetDynamics_Execute()
 	{
 		return;
 	}
-	if (FRigPhysicsBodyComponent* Component = Cast<FRigPhysicsBodyComponent>(
-		ExecuteContext.Hierarchy->FindComponent(PhysicsBodyComponentKey)))
+	if (FRigPhysicsBodyComponent* Component = GetBody(ExecuteContext.Hierarchy, PhysicsBodyComponentKey))
 	{
 		Component->Dynamics = Dynamics;
 	}
@@ -94,8 +97,7 @@ FRigUnit_HierarchySetCollision_Execute()
 	{
 		return;
 	}
-	if (FRigPhysicsBodyComponent* Component = Cast<FRigPhysicsBodyComponent>(
-		ExecuteContext.Hierarchy->FindComponent(PhysicsBodyComponentKey)))
+	if (FRigPhysicsBodyComponent* Component = GetBody(ExecuteContext.Hierarchy, PhysicsBodyComponentKey))
 	{
 		Component->Collision = Collision;
 	}
@@ -109,8 +111,7 @@ FRigUnit_HierarchyDisableCollisionBetween_Execute()
 	{
 		return;
 	}
-	if (FRigPhysicsBodyComponent* Component1 = Cast<FRigPhysicsBodyComponent>(
-		ExecuteContext.Hierarchy->FindComponent(PhysicsBodyComponentKey1)))
+	if (FRigPhysicsBodyComponent* Component1 = GetBody(ExecuteContext.Hierarchy, PhysicsBodyComponentKey1))
 	{
 		Component1->NoCollisionBodies.Add(PhysicsBodyComponentKey2);
 	}
@@ -125,8 +126,7 @@ FRigUnit_HierarchySetPhysicsBodySourceBone_Execute()
 	{
 		return;
 	}
-	if (FRigPhysicsBodyComponent* Component = Cast<FRigPhysicsBodyComponent>(
-		ExecuteContext.Hierarchy->FindComponent(PhysicsBodyComponentKey)))
+	if (FRigPhysicsBodyComponent* Component = GetBody(ExecuteContext.Hierarchy, PhysicsBodyComponentKey))
 	{
 		Component->BodySolverSettings.SourceBone = SourceBone;
 	}
@@ -140,8 +140,7 @@ FRigUnit_HierarchySetPhysicsBodyTargetBone_Execute()
 	{
 		return;
 	}
-	if (FRigPhysicsBodyComponent* Component = Cast<FRigPhysicsBodyComponent>(
-		ExecuteContext.Hierarchy->FindComponent(PhysicsBodyComponentKey)))
+	if (FRigPhysicsBodyComponent* Component = GetBody(ExecuteContext.Hierarchy, PhysicsBodyComponentKey))
 	{
 		Component->BodySolverSettings.TargetBone = TargetBone;
 	}
@@ -155,8 +154,7 @@ FRigUnit_HierarchySetPhysicsBodyKinematicTarget_Execute()
 	{
 		return;
 	}
-	if (FRigPhysicsBodyComponent* Component = Cast<FRigPhysicsBodyComponent>(
-		ExecuteContext.Hierarchy->FindComponent(PhysicsBodyComponentKey)))
+	if (FRigPhysicsBodyComponent* Component = GetBody(ExecuteContext.Hierarchy, PhysicsBodyComponentKey))
 	{
 		Component->KinematicTarget = KinematicTarget;
 	}
@@ -170,8 +168,7 @@ FRigUnit_HierarchySetPhysicsBodyKinematicTargetSpace_Execute()
 	{
 		return;
 	}
-	if (FRigPhysicsBodyComponent* Component = Cast<FRigPhysicsBodyComponent>(
-		ExecuteContext.Hierarchy->FindComponent(PhysicsBodyComponentKey)))
+	if (FRigPhysicsBodyComponent* Component = GetBody(ExecuteContext.Hierarchy, PhysicsBodyComponentKey))
 	{
 		Component->BodyData.KinematicTargetSpace = KinematicTargetSpace;
 	}
@@ -188,8 +185,7 @@ FRigUnit_HierarchySetPhysicsBodySparseData_Execute()
 	{
 		return;
 	}
-	if (FRigPhysicsBodyComponent* Component = Cast<FRigPhysicsBodyComponent>(
-		ExecuteContext.Hierarchy->FindComponent(PhysicsBodyComponentKey)))
+	if (FRigPhysicsBodyComponent* Component = GetBody(ExecuteContext.Hierarchy, PhysicsBodyComponentKey))
 	{
 		Component->BodyData.UpdateFromSparseData(Data);
 	}
@@ -203,8 +199,7 @@ FRigUnit_HierarchySetPhysicsBodyMovementType_Execute()
 	{
 		return;
 	}
-	if (FRigPhysicsBodyComponent* Component = Cast<FRigPhysicsBodyComponent>(
-		ExecuteContext.Hierarchy->FindComponent(PhysicsBodyComponentKey)))
+	if (FRigPhysicsBodyComponent* Component = GetBody(ExecuteContext.Hierarchy, PhysicsBodyComponentKey))
 	{
 		Component->BodyData.MovementType = MovementType;
 	}
@@ -218,8 +213,7 @@ FRigUnit_HierarchySetPhysicsBodyCollisionType_Execute()
 	{
 		return;
 	}
-	if (FRigPhysicsBodyComponent* Component = Cast<FRigPhysicsBodyComponent>(
-		ExecuteContext.Hierarchy->FindComponent(PhysicsBodyComponentKey)))
+	if (FRigPhysicsBodyComponent* Component = GetBody(ExecuteContext.Hierarchy, PhysicsBodyComponentKey))
 	{
 		Component->BodyData.CollisionType = CollisionType;
 	}
@@ -233,8 +227,7 @@ FRigUnit_HierarchySetPhysicsBodyIncludeInChecksForReset_Execute()
 	{
 		return;
 	}
-	if (FRigPhysicsBodyComponent* Component = Cast<FRigPhysicsBodyComponent>(
-		ExecuteContext.Hierarchy->FindComponent(PhysicsBodyComponentKey)))
+	if (FRigPhysicsBodyComponent* Component = GetBody(ExecuteContext.Hierarchy, PhysicsBodyComponentKey))
 	{
 		Component->BodySolverSettings.bIncludeInChecksForReset = bInclude;
 	}
@@ -253,14 +246,12 @@ FRigUnit_HierarchySetPhysicsBodyMaterial_Execute()
 		UE_CONTROLRIG_RIGUNIT_REPORT_ERROR(TEXT("SetPhysicsBodyMaterial can only be used during Setup"));
 	}
 
-	if (FRigPhysicsBodyComponent* Component = Cast<FRigPhysicsBodyComponent>(
-		ExecuteContext.Hierarchy->FindComponent(PhysicsBodyComponentKey)))
+	if (FRigPhysicsBodyComponent* Component = GetBody(ExecuteContext.Hierarchy, PhysicsBodyComponentKey))
 	{
 		Component->Collision.Material = Material;
 	}
 }
 
-
 //======================================================================================================================
 FRigUnit_HierarchySetPhysicsBodyGravityMultiplier_Execute()
 {
@@ -269,8 +260,7 @@ FRigUnit_HierarchySetPhysicsBodyGravityMultiplier_Execute()
 	{
 		return;
 	}
-	if (FRigPhysicsBodyComponent* Component = Cast<FRigPhysicsBodyComponent>(
-		ExecuteContext.Hierarchy->FindComponent(PhysicsBodyComponentKey)))
+	if (FRigPhysicsBodyComponent* Component = GetBody(ExecuteContext.Hierarchy, PhysicsBodyComponentKey))
 	{
 		Component->BodyData.GravityMultiplier = GravityMultiplier;
 	}
@@ -286,8 +276,7 @@ FRigUnit_HierarchySetPhysicsBodyPhysicsBlendWeight_Execute()
 	{
 		return;
 	}
-	if (FRigPhysicsBodyComponent* Component = Cast<FRigPhysicsBodyComponent>(
-		ExecuteContext.Hierarchy->FindComponent(PhysicsBodyComponentKey)))
+	if (FRigPhysicsBodyComponent* Component = GetBody(ExecuteContext.Hierarchy, PhysicsBodyComponentKey))
 	{
 		Component->BodyData.PhysicsBlendWeight = PhysicsBlendWeight;
 	}
@@ -301,8 +290,7 @@ FRigUnit_HierarchySetPhysicsBodyUpdateKinematicFromSimulation_Execute()
 	{
 		return;
 	}
-	if (FRigPhysicsBodyComponent* Component = Cast<FRigPhysicsBodyComponent>(
-		ExecuteContext.Hierarchy->FindComponent(PhysicsBodyComponentKey)))
+	if (FRigPhysicsBodyComponent* Component = GetBody(ExecuteContext.Hierarchy, PhysicsBodyComponentKey))
 	{
 		Component->BodyData.bUpdateKinematicFromSimulation = bUpdateKinematicFromSimulation;
 	}
@@ -316,10 +304,131 @@ FRigUnit_HierarchySetPhysicsBodyDamping_Execute()
 	{
 		return;
 	}
-	if (FRigPhysicsBodyComponent* Component = Cast<FRigPhysicsBodyComponent>(
-		ExecuteContext.Hierarchy->FindComponent(PhysicsBodyComponentKey)))
+	if (FRigPhysicsBodyComponent* Component = GetBody(ExecuteContext.Hierarchy, PhysicsBodyComponentKey))
 	{
 		Component->Dynamics.LinearDamping = LinearDamping;
 		Component->Dynamics.AngularDamping = AngularDamping;
 	}
 }
+
+//======================================================================================================================
+FRigUnit_HierarchyAddPhysicsBodyForce_Execute()
+{
+	DECLARE_SCOPE_HIERARCHICAL_COUNTER_RIGUNIT();
+	if (!ExecuteContext.Hierarchy)
+	{
+		return;
+	}
+	if (FRigPhysicsBodyComponent* Component = GetBody(ExecuteContext.Hierarchy, PhysicsBodyComponentKey))
+	{
+		Component->ForceAndTorques.Add(FPhysicsControlNamedForceAndTorqueData(Name, ForceAndTorque));
+	}
+}
+
+//======================================================================================================================
+FRigUnit_HierarchyRemovePhysicsBodyForce_Execute()
+{
+	DECLARE_SCOPE_HIERARCHICAL_COUNTER_RIGUNIT();
+	if (!ExecuteContext.Hierarchy)
+	{
+		return;
+	}
+	if (FRigPhysicsBodyComponent* Component = GetBody(ExecuteContext.Hierarchy, PhysicsBodyComponentKey))
+	{
+		Component->ForceAndTorques.RemoveAll([Name](const FPhysicsControlNamedForceAndTorqueData& Data) 
+			{
+				return Data.Name == Name;
+			});
+	}
+}
+
+//======================================================================================================================
+FRigUnit_HierarchyGetPhysicsBodyTransform_Execute()
+{
+	DECLARE_SCOPE_HIERARCHICAL_COUNTER_RIGUNIT();
+	if (!ExecuteContext.Hierarchy)
+	{
+		return;
+	}
+	if (const FRigPhysicsBodyComponent* Component = GetBody(ExecuteContext.Hierarchy, PhysicsBodyComponentKey))
+	{
+		Transform = Component->Transform;
+	}
+}
+
+//======================================================================================================================
+FRigUnit_HierarchyGetPhysicsBodyCoMTransform_Execute()
+{
+	DECLARE_SCOPE_HIERARCHICAL_COUNTER_RIGUNIT();
+	if (!ExecuteContext.Hierarchy)
+	{
+		return;
+	}
+	if (const FRigPhysicsBodyComponent* Component = GetBody(ExecuteContext.Hierarchy, PhysicsBodyComponentKey))
+	{
+		CoMTransform = Component->CoMTransform;
+	}
+}
+
+//======================================================================================================================
+FRigUnit_HierarchyGetPhysicsBodyLinearVelocity_Execute()
+{
+	DECLARE_SCOPE_HIERARCHICAL_COUNTER_RIGUNIT();
+	if (!ExecuteContext.Hierarchy)
+	{
+		return;
+	}
+	if (const FRigPhysicsBodyComponent* Component = GetBody(ExecuteContext.Hierarchy, PhysicsBodyComponentKey))
+	{
+		LinearVelocity = Component->LinearVelocity;
+	}
+}
+
+//======================================================================================================================
+FRigUnit_HierarchyGetPhysicsBodyAngularVelocity_Execute()
+{
+	DECLARE_SCOPE_HIERARCHICAL_COUNTER_RIGUNIT();
+	if (!ExecuteContext.Hierarchy)
+	{
+		return;
+	}
+	if (const FRigPhysicsBodyComponent* Component = GetBody(ExecuteContext.Hierarchy, PhysicsBodyComponentKey))
+	{
+		AngularVelocity = Component->AngularVelocity;
+	}
+}
+
+//======================================================================================================================
+FRigUnit_HierarchyGetPhysicsBodyPointVelocity_Execute()
+{
+	DECLARE_SCOPE_HIERARCHICAL_COUNTER_RIGUNIT();
+	if (!ExecuteContext.Hierarchy)
+	{
+		return;
+	}
+	if (const FRigPhysicsBodyComponent* Component = GetBody(ExecuteContext.Hierarchy, PhysicsBodyComponentKey))
+	{
+		Velocity = Component->LinearVelocity;
+		FVector AngularVelocity = Component->AngularVelocity;
+		FVector Location;
+		switch (Space)
+		{
+			// For control rig, World = Component space
+			case EPhysicsControlSpace::World:
+			case EPhysicsControlSpace::Component:
+				Location = Position;
+				break;
+			case EPhysicsControlSpace::Body:
+				Location = Component->Transform.TransformPositionNoScale(Position);
+				break;
+			default:
+				Location = Position;
+				break;
+		}
+
+		FVector Delta = Location - Component->CoMTransform.GetLocation();
+		Velocity += FVector::CrossProduct(AngularVelocity, Delta);
+	}
+}
+
+
diff --git a/Engine/Plugins/Experimental/ControlRigPhysics/Source/ControlRigPhysics/Private/RigPhysicsData.cpp b/Engine/Plugins/Experimental/ControlRigPhysics/Source/ControlRigPhysics/Private/RigPhysicsData.cpp
--- a/Engine/Plugins/Experimental/ControlRigPhysics/Source/ControlRigPhysics/Private/RigPhysicsData.cpp
+++ b/Engine/Plugins/Experimental/ControlRigPhysics/Source/ControlRigPhysics/Private/RigPhysicsData.cpp
@@ -270,9 +270,9 @@ FArchive& operator <<(FArchive& Ar, FRigPhysicsSolverSettings& Data)
 			Ar << ResetCooldownFrames;
 		}
 	}
-	if (Ar.CustomVer(FPhysicsControlObjectVersion::GUID) >= FPhysicsControlObjectVersion::ControlRigAutomaticallyAddPhysicsComponents)
+	if (Ar.CustomVer(FPhysicsControlObjectVersion::GUID) >= FPhysicsControlObjectVersion::ControlRigAutomaticallyAddPhysicsBodyComponents)
 	{
-		Ar << Data.bAutomaticallyAddPhysicsComponents;
+		Ar << Data.bAutomaticallyAddPhysicsBodyComponents;
 	}
 	if (Ar.CustomVer(FPhysicsControlObjectVersion::GUID) >= FPhysicsControlObjectVersion::ControlRigWorldCollision)
 	{
diff --git a/Engine/Plugins/Experimental/ControlRigPhysics/Source/ControlRigPhysics/Private/RigPhysicsSimulation.cpp b/Engine/Plugins/Experimental/ControlRigPhysics/Source/ControlRigPhysics/Private/RigPhysicsSimulation.cpp
--- a/Engine/Plugins/Experimental/ControlRigPhysics/Source/ControlRigPhysics/Private/RigPhysicsSimulation.cpp
+++ b/Engine/Plugins/Experimental/ControlRigPhysics/Source/ControlRigPhysics/Private/RigPhysicsSimulation.cpp
@@ -78,30 +78,30 @@ FRigPhysicsSimulation::FRigPhysicsSimulation(const FName InOwnerName)
 bool FRigPhysicsSimulation::ShouldComponentBeInSimulation(
 	const URigHierarchy& Hierarchy, FRigComponentKey SolverComponentKey, FRigComponentKey ComponentKey) const
 {
-	const FRigPhysicsBodyComponent* PhysicsComponent = Cast<FRigPhysicsBodyComponent>(
+	const FRigPhysicsBodyComponent* PhysicsBodyComponent = Cast<FRigPhysicsBodyComponent>(
 		Hierarchy.FindComponent(ComponentKey));
 
 	const FRigPhysicsSolverComponent* PhysicsSolverComponent = Cast<FRigPhysicsSolverComponent>(
 		Hierarchy.FindComponent(SolverComponentKey));
 
-	if (!PhysicsComponent || !PhysicsSolverComponent)
+	if (!PhysicsBodyComponent || !PhysicsSolverComponent)
 	{
 		return false;
 	}
 
-	if (PhysicsComponent->BodySolverSettings.PhysicsSolverComponentKey == SolverComponentKey)
+	if (PhysicsBodyComponent->BodySolverSettings.PhysicsSolverComponentKey == SolverComponentKey)
 	{
 		return true;
 	}
 
-	if (!PhysicsSolverComponent->SolverSettings.bAutomaticallyAddPhysicsComponents)
+	if (!PhysicsSolverComponent->SolverSettings.bAutomaticallyAddPhysicsBodyComponents)
 	{
 		return false;
 	}
 
-	if (PhysicsComponent->BodySolverSettings.bUseAutomaticSolver)
+	if (PhysicsBodyComponent->BodySolverSettings.bUseAutomaticSolver)
 	{
-		FRigElementKey ElementKey = PhysicsComponent->GetElementKey();
+		FRigElementKey ElementKey = PhysicsBodyComponent->GetElementKey();
 		while (ElementKey.IsValid())
 		{
 			for (FRigComponentKey CK : Hierarchy.GetComponentKeys(ElementKey))
@@ -618,7 +618,7 @@ ImmediatePhysics::FActorHandle* FRigPhysicsSimulation::CreateBody(
 //======================================================================================================================
 void FRigPhysicsSimulation::Instantiate(
 	const FRigVMExecuteContext&       ExecuteContext, 
-	const URigHierarchy&              Hierarchy, 
+	URigHierarchy&                    Hierarchy, 
 	const FRigPhysicsSolverComponent* SolverComponent)
 {
 	DECLARE_SCOPE_HIERARCHICAL_COUNTER_FUNC()
@@ -710,7 +710,7 @@ void FRigPhysicsSimulation::InstantiateSolverCollision(
 
 //======================================================================================================================
 void FRigPhysicsSimulation::InstantiatePhysicsBodies(
-	const URigHierarchy&              Hierarchy,
+	URigHierarchy&                    Hierarchy,
 	const FRigPhysicsSolverComponent* SolverComponent,
 	FRigPhysicsIgnorePairs&           IgnorePairs)
 {
@@ -725,14 +725,14 @@ void FRigPhysicsSimulation::InstantiatePhysicsBodies(
 		const FRigComponentKey& ComponentKey = BodyRecordPair.Key;
 		FRigBodyRecord& Record = BodyRecordPair.Value;
 
-		if (const FRigPhysicsBodyComponent* PhysicsComponent = Cast<FRigPhysicsBodyComponent>(
+		if (FRigPhysicsBodyComponent* PhysicsBodyComponent = Cast<FRigPhysicsBodyComponent>(
 			Hierarchy.FindComponent(ComponentKey)))
 		{
-			const FRigPhysicsCollision& Collision = PhysicsComponent->Collision;
-			const FRigPhysicsDynamics& Dynamics = PhysicsComponent->Dynamics;
-			const FPhysicsControlModifierData& BodyData = PhysicsComponent->BodyData;
+			const FRigPhysicsCollision& Collision = PhysicsBodyComponent->Collision;
+			const FRigPhysicsDynamics& Dynamics = PhysicsBodyComponent->Dynamics;
+			const FPhysicsControlModifierData& BodyData = PhysicsBodyComponent->BodyData;
 
-			FRigElementKey SourceKey = PhysicsComponent->BodySolverSettings.SourceBone;
+			FRigElementKey SourceKey = PhysicsBodyComponent->BodySolverSettings.SourceBone;
 			if (!SourceKey.IsValid())
 			{
 				SourceKey = ComponentKey.ElementKey;
@@ -741,20 +741,24 @@ void FRigPhysicsSimulation::InstantiatePhysicsBodies(
 			// What should we do if the key is not valid? 
 			if (SourceKey.IsValid())
 			{
-				FTransform SourceComponentSpaceTM = Hierarchy.GetGlobalTransform(SourceKey);
+				Record.FinalComponentSpaceTM = Hierarchy.GetGlobalTransform(SourceKey);
 				const FTransform SourceSimulationSpaceTM =
-					ConvertComponentSpaceTransformToSimSpace(SolverComponent->SolverSettings, SourceComponentSpaceTM);
+					ConvertComponentSpaceTransformToSimSpace(SolverComponent->SolverSettings, Record.FinalComponentSpaceTM);
 				Record.ActorHandle = CreateBody(
 					ComponentKey.ElementKey.Name, Collision, &Dynamics, &BodyData, SourceSimulationSpaceTM);
+				if (Record.ActorHandle)
+				{
+					SetPhysicsBodyComponentState(SolverSettings, Record, *PhysicsBodyComponent);
+				}
 			}
 
-			Record.TargetElementKey = PhysicsComponent->BodySolverSettings.TargetBone;
+			Record.TargetElementKey = PhysicsBodyComponent->BodySolverSettings.TargetBone;
 			if (!Record.TargetElementKey.IsValid())
 			{
 				Record.TargetElementKey = ComponentKey.ElementKey;
 			}
 
-			for (const FRigComponentKey& NoCollisionKey : PhysicsComponent->NoCollisionBodies)
+			for (const FRigComponentKey& NoCollisionKey : PhysicsBodyComponent->NoCollisionBodies)
 			{
 				IgnorePairs.Add(FRigPhysicsIgnorePair(ComponentKey, NoCollisionKey));
 			}
@@ -762,6 +766,20 @@ void FRigPhysicsSimulation::InstantiatePhysicsBodies(
 	}
 }
 
+//======================================================================================================================
+void FRigPhysicsSimulation::SetPhysicsBodyComponentState(
+	const FRigPhysicsSolverSettings& SolverSettings,
+	FRigBodyRecord&                  Record, 
+	FRigPhysicsBodyComponent&        PhysicsBodyComponent) const
+{
+	PhysicsBodyComponent.Transform = Record.FinalComponentSpaceTM;
+	PhysicsBodyComponent.CoMTransform = Record.ActorHandle->GetLocalCoMTransform() * Record.FinalComponentSpaceTM;
+	PhysicsBodyComponent.LinearVelocity = 
+		ConvertSimSpaceVectorToComponentSpace(SolverSettings, Record.ActorHandle->GetLinearVelocity());
+	PhysicsBodyComponent.AngularVelocity = 
+		ConvertSimSpaceVectorToComponentSpace(SolverSettings, Record.ActorHandle->GetAngularVelocity());
+}
+
 //======================================================================================================================
 void FRigPhysicsSimulation::InstantiatePhysicsJoints(
 	const URigHierarchy&              Hierarchy,
@@ -792,7 +810,7 @@ void FRigPhysicsSimulation::InstantiatePhysicsJoints(
 			FRigElementKey ParentBoneKey = ParentBodyComponentKey.ElementKey;
 
 			// Joints require both parent and child to exist
-			const FRigPhysicsBodyComponent* ChildPhysicsComponent =
+			const FRigPhysicsBodyComponent* ChildPhysicsBodyComponent =
 				Cast<FRigPhysicsBodyComponent>(Hierarchy.FindComponent(ChildBodyComponentKey));
 			FRigBodyRecord* ChildBodyRecord = BodyRecords.Find(ChildBodyComponentKey);
 			ImmediatePhysics::FActorHandle* ChildActorHandle = ChildBodyRecord ? ChildBodyRecord->ActorHandle : nullptr;
@@ -801,7 +819,7 @@ void FRigPhysicsSimulation::InstantiatePhysicsJoints(
 				continue;
 			}
 
-			const FRigPhysicsBodyComponent* ParentPhysicsComponent =
+			const FRigPhysicsBodyComponent* ParentPhysicsBodyComponent =
 				Cast<FRigPhysicsBodyComponent>(Hierarchy.FindComponent(ParentBodyComponentKey));
 			FRigBodyRecord* ParentBodyRecord = BodyRecords.Find(ParentBodyComponentKey);
 			ImmediatePhysics::FActorHandle* ParentActorHandle = ParentBodyRecord ? ParentBodyRecord->ActorHandle : nullptr;
@@ -915,20 +933,20 @@ void FRigPhysicsSimulation::UpdateBodyRecordPrePhysics(
 	const FRigPhysicsSolverComponent* SolverComponent,
 	const float                       DeltaTime,
 	FRigBodyRecord&                   Record, 
-	const FRigPhysicsBodyComponent*   PhysicsComponent)
+	const FRigPhysicsBodyComponent*   PhysicsBodyComponent)
 {
 	if (!Record.ActorHandle)
 	{
 		return;
 	}
-	const FRigComponentKey ComponentKey = PhysicsComponent->GetKey();
+	const FRigComponentKey ComponentKey = PhysicsBodyComponent->GetKey();
 
 	// Shuffle the record data
 	Record.PrevSourceComponentSpaceVelocity = Record.SourceComponentSpaceVelocity;
 	Record.PrevSourceComponentSpaceAngularVelocity = Record.SourceComponentSpaceAngularVelocity;
 	Record.PrevSourceComponentSpaceTM = Record.SourceComponentSpaceTM;
 
-	FRigElementKey SourceKey = PhysicsComponent->BodySolverSettings.SourceBone;
+	FRigElementKey SourceKey = PhysicsBodyComponent->BodySolverSettings.SourceBone;
 	if (!SourceKey.IsValid())
 	{
 		SourceKey = ComponentKey.ElementKey;
@@ -958,11 +976,11 @@ void FRigPhysicsSimulation::UpdateBodyPrePhysics(
 	const FRigVMExecuteContext&       ExecuteContext, 
 	const FRigPhysicsSolverComponent* SolverComponent,
 	const FRigBodyRecord&             Record, 
-	const FRigPhysicsBodyComponent*   PhysicsComponent)
+	const FRigPhysicsBodyComponent*   PhysicsBodyComponent)
 {
-	const FRigComponentKey ComponentKey = PhysicsComponent->GetKey();
+	const FRigComponentKey ComponentKey = PhysicsBodyComponent->GetKey();
 
-	FPhysicsControlModifierData BodyData = PhysicsComponent->BodyData;
+	FPhysicsControlModifierData BodyData = PhysicsBodyComponent->BodyData;
 	EPhysicsControlKinematicTargetSpace KinematicTargetSpace = BodyData.KinematicTargetSpace;
 
 	if (SolverComponent->TrackInputCounter > 0)
@@ -989,13 +1007,13 @@ void FRigPhysicsSimulation::UpdateBodyPrePhysics(
 		{
 		case EPhysicsControlKinematicTargetSpace::Component:
 		{
-			KinematicTargetCS = PhysicsComponent->KinematicTarget;
+			KinematicTargetCS = PhysicsBodyComponent->KinematicTarget;
 			break;
 		}
 		case EPhysicsControlKinematicTargetSpace::World:
 		{
 			// Record.KinematicTarget * SimulationSpaceState.ComponentTM.Inverse();
-			KinematicTargetCS = PhysicsComponent->KinematicTarget.GetRelativeTransform(
+			KinematicTargetCS = PhysicsBodyComponent->KinematicTarget.GetRelativeTransform(
 				SimulationSpaceState.ComponentTM);
 			break;
 		}
@@ -1003,7 +1021,7 @@ void FRigPhysicsSimulation::UpdateBodyPrePhysics(
 		{
 			// All the other options are relative to a bone, so the first task is to get
 			// that, which will be in component space
-			FRigElementKey SourceKey = PhysicsComponent->BodySolverSettings.SourceBone;
+			FRigElementKey SourceKey = PhysicsBodyComponent->BodySolverSettings.SourceBone;
 			if (!SourceKey.IsValid())
 			{
 				SourceKey = ComponentKey.ElementKey;
@@ -1014,7 +1032,7 @@ void FRigPhysicsSimulation::UpdateBodyPrePhysics(
 				{
 				case EPhysicsControlKinematicTargetSpace::OffsetInBoneSpace:
 				{
-					KinematicTargetCS = PhysicsComponent->KinematicTarget * Record.SourceComponentSpaceTM.ToTransform();
+					KinematicTargetCS = PhysicsBodyComponent->KinematicTarget * Record.SourceComponentSpaceTM.ToTransform();
 					break;
 				}
 				case EPhysicsControlKinematicTargetSpace::OffsetInWorldSpace:
@@ -1024,9 +1042,9 @@ void FRigPhysicsSimulation::UpdateBodyPrePhysics(
 					FTransform KinematicTargetWS = BoneWS;
 					// Note that we don't want to rotate the translation by the target - we apply it
 					// individually to orientation and position.
-					KinematicTargetWS.AddToTranslation(PhysicsComponent->KinematicTarget.GetTranslation());
+					KinematicTargetWS.AddToTranslation(PhysicsBodyComponent->KinematicTarget.GetTranslation());
 					KinematicTargetWS.SetRotation(
-						PhysicsComponent->KinematicTarget.GetRotation() * KinematicTargetWS.GetRotation());
+						PhysicsBodyComponent->KinematicTarget.GetRotation() * KinematicTargetWS.GetRotation());
 					// Danny TODO figure out which of the GetRelativeTransform versions this is
 					KinematicTargetCS = SimulationSpaceState.ComponentTM.Inverse() * KinematicTargetWS;
 					break;
@@ -1036,7 +1054,7 @@ void FRigPhysicsSimulation::UpdateBodyPrePhysics(
 					// This applies the offset as a transform in component space. Note that this is
 					// different to what PCC does, because in ControlRig "world" is really
 					// "component" space.
-					KinematicTargetCS = Record.SourceComponentSpaceTM.ToTransform() * PhysicsComponent->KinematicTarget;
+					KinematicTargetCS = Record.SourceComponentSpaceTM.ToTransform() * PhysicsBodyComponent->KinematicTarget;
 					break;
 				}
 				case EPhysicsControlKinematicTargetSpace::IgnoreTarget:
@@ -1060,8 +1078,79 @@ void FRigPhysicsSimulation::UpdateBodyPrePhysics(
 	else
 	{
 		// Danny TODO move damping into BodyData - any PhysicsControl system should be able to use it
-		Record.ActorHandle->SetLinearDamping(PhysicsComponent->Dynamics.LinearDamping);
-		Record.ActorHandle->SetAngularDamping(PhysicsComponent->Dynamics.AngularDamping);
+		Record.ActorHandle->SetLinearDamping(PhysicsBodyComponent->Dynamics.LinearDamping);
+		Record.ActorHandle->SetAngularDamping(PhysicsBodyComponent->Dynamics.AngularDamping);
+	}
+}
+
+//======================================================================================================================
+void FRigPhysicsSimulation::ApplyForcesPrePhysics(
+	const FRigVMExecuteContext&       ExecuteContext,
+	const FRigPhysicsSolverComponent* SolverComponent,
+	const FRigBodyRecord&             Record,
+	const FRigPhysicsBodyComponent*   PhysicsBodyComponent)
+{
+	if (PhysicsBodyComponent->ForceAndTorques.IsEmpty() || !Record.ActorHandle || Record.ActorHandle->GetIsKinematic())
+	{
+		return;
+	}
+
+	ImmediatePhysics::FActorHandle* ActorHandle = Record.ActorHandle;
+	const FRigPhysicsSolverSettings& SolverSettings = SolverComponent->SolverSettings;
+
+	for (const FPhysicsControlNamedForceAndTorqueData& Data : PhysicsBodyComponent->ForceAndTorques)
+	{
+		const FPhysicsControlForceAndTorqueData& ForceAndTorqueData = Data.ForceAndTorqueData;
+
+		// Need to convert the requested force/torque into simulation space first
+		FVector LinearSS;
+		FVector AngularSS;
+
+		FTransform BodyTM = ActorHandle->GetWorldTransform();
+
+		switch (ForceAndTorqueData.ForceAndTorqueSpace)
+		{
+		case EPhysicsControlSpace::World: // For Control Rig, World and Component space are considered to be the same
+		case EPhysicsControlSpace::Component:
+			LinearSS = ConvertComponentSpaceVectorToSimSpace(
+				SolverSettings, SimulationSpaceState.ComponentTM.InverseTransformVector(ForceAndTorqueData.Linear));
+			AngularSS = ConvertComponentSpaceVectorToSimSpace(
+				SolverSettings, SimulationSpaceState.ComponentTM.InverseTransformVector(ForceAndTorqueData.Angular));
+			break;
+		case EPhysicsControlSpace::Body:
+			LinearSS = BodyTM.GetRotation() * ForceAndTorqueData.Linear;
+			AngularSS = BodyTM.GetRotation() * ForceAndTorqueData.Angular;
+			break;
+		default:
+			LinearSS = FVector::ZeroVector;
+			AngularSS = FVector::ZeroVector;
+			break;
+		}
+
+		// Apply it differently depending on the location
+		ImmediatePhysics_Shared::EForceType FT = ConvertForceType(ForceAndTorqueData.Type);
+		switch (ForceAndTorqueData.LocationSpace)
+		{
+		case EPhysicsControlSpace::World: // For Control Rig, World and Component space are considered to be the same
+		case EPhysicsControlSpace::Component:
+		{
+			FVector SimSpaceLocation = ConvertWorldPositionToSimSpaceNoScale(
+				SolverSettings, ForceAndTorqueData.Location);
+			ActorHandle->AddForceAtLocation(LinearSS, SimSpaceLocation, FT);
+		}
+		break;
+		case EPhysicsControlSpace::Body:
+		{
+			FVector SimSpaceLocation = ActorHandle->GetWorldTransform().TransformPositionNoScale(
+				ActorHandle->GetLocalCoMLocation() + ForceAndTorqueData.Location);
+			ActorHandle->AddForceAtLocation(LinearSS, SimSpaceLocation, FT);
+		}
+		break;
+		default:
+			break;
+		}
+		// Torque doesn't get applied at a location
+		ActorHandle->AddTorque(AngularSS, FT);
 	}
 }
 
@@ -1111,10 +1200,10 @@ void FRigPhysicsSimulation::UpdateJointPrePhysics(
 			float TargetVelocityMultiplier = 1.0f;
 
 			FRigElementKey ChildSourceKey;
-			if (const FRigPhysicsBodyComponent* ChildPhysicsComponent = Cast<FRigPhysicsBodyComponent>(
+			if (const FRigPhysicsBodyComponent* ChildPhysicsBodyComponent = Cast<FRigPhysicsBodyComponent>(
 				Hierarchy.FindComponent(Record.ChildBodyComponentKey)))
 			{
-				ChildSourceKey = ChildPhysicsComponent->BodySolverSettings.SourceBone;
+				ChildSourceKey = ChildPhysicsBodyComponent->BodySolverSettings.SourceBone;
 				if (!ChildSourceKey.IsValid())
 				{
 					ChildSourceKey = ComponentKey.ElementKey;
@@ -1126,10 +1215,10 @@ void FRigPhysicsSimulation::UpdateJointPrePhysics(
 			if (DriveData.bUseSkeletalAnimation)
 			{
 				FRigElementKey ParentSourceKey;
-				if (const FRigPhysicsBodyComponent* ParentPhysicsComponent = Cast<FRigPhysicsBodyComponent>(
+				if (const FRigPhysicsBodyComponent* ParentPhysicsBodyComponent = Cast<FRigPhysicsBodyComponent>(
 					Hierarchy.FindComponent(Record.ParentBodyComponentKey)))
 				{
-					ParentSourceKey = ParentPhysicsComponent->BodySolverSettings.SourceBone;
+					ParentSourceKey = ParentPhysicsBodyComponent->BodySolverSettings.SourceBone;
 					if (!ParentSourceKey.IsValid())
 					{
 						ParentSourceKey = Record.ParentBodyComponentKey.ElementKey;
@@ -1250,10 +1339,10 @@ void FRigPhysicsSimulation::CheckForResetsPrePhysics(
 
 		if (Record.ActorHandle && Record.ActorHandle->GetIsKinematic())
 		{
-			if (const FRigPhysicsBodyComponent* PhysicsComponent = Cast<FRigPhysicsBodyComponent>(
+			if (const FRigPhysicsBodyComponent* PhysicsBodyComponent = Cast<FRigPhysicsBodyComponent>(
 				Hierarchy.FindComponent(ComponentKey)))
 			{
-				if (PhysicsComponent->BodySolverSettings.bIncludeInChecksForReset)
+				if (PhysicsBodyComponent->BodySolverSettings.bIncludeInChecksForReset)
 				{
 					const FVector Velocity = Record.SourceComponentSpaceVelocity;
 					const FVector Acceleration =
@@ -1304,10 +1393,10 @@ void FRigPhysicsSimulation::UpdatePrePhysics(
 	{
 		const FRigComponentKey& ComponentKey = BodyRecordPair.Key;
 		FRigBodyRecord& Record = BodyRecordPair.Value;
-		if (const FRigPhysicsBodyComponent* PhysicsComponent = Cast<FRigPhysicsBodyComponent>(
+		if (const FRigPhysicsBodyComponent* PhysicsBodyComponent = Cast<FRigPhysicsBodyComponent>(
 			Hierarchy.FindComponent(ComponentKey)))
 		{
-			UpdateBodyRecordPrePhysics(Hierarchy, SolverComponent, DeltaTime, Record, PhysicsComponent);
+			UpdateBodyRecordPrePhysics(Hierarchy, SolverComponent, DeltaTime, Record, PhysicsBodyComponent);
 		}
 	}
 
@@ -1325,10 +1414,10 @@ void FRigPhysicsSimulation::UpdatePrePhysics(
 		{
 			const FRigComponentKey& ComponentKey = BodyRecordPair.Key;
 			FRigBodyRecord& Record = BodyRecordPair.Value;
-			if (const FRigPhysicsBodyComponent* PhysicsComponent = Cast<FRigPhysicsBodyComponent>(
+			if (const FRigPhysicsBodyComponent* PhysicsBodyComponent = Cast<FRigPhysicsBodyComponent>(
 				Hierarchy.FindComponent(ComponentKey)))
 			{
-				UpdateBodyPrePhysics(ExecuteContext, SolverComponent, Record, PhysicsComponent);
+				UpdateBodyPrePhysics(ExecuteContext, SolverComponent, Record, PhysicsBodyComponent);
 			}
 		}
 	}
@@ -1368,6 +1457,22 @@ void FRigPhysicsSimulation::UpdatePrePhysics(
 			}
 		}
 	}
+
+	{
+		DECLARE_SCOPE_HIERARCHICAL_COUNTER_FUNC()
+		QUICK_SCOPE_CYCLE_COUNTER(STAT_RigPhysics_ApplyForcesPrePhysics);
+		for (TPair<FRigComponentKey, FRigBodyRecord>& BodyRecordPair : BodyRecords)
+		{
+			const FRigComponentKey& ComponentKey = BodyRecordPair.Key;
+			FRigBodyRecord& Record = BodyRecordPair.Value;
+			if (const FRigPhysicsBodyComponent* PhysicsBodyComponent = Cast<FRigPhysicsBodyComponent>(
+				Hierarchy.FindComponent(ComponentKey)))
+			{
+				ApplyForcesPrePhysics(ExecuteContext, SolverComponent, Record, PhysicsBodyComponent);
+			}
+
+		}
+	}
 }
 
 //======================================================================================================================
@@ -1522,6 +1627,7 @@ void FRigPhysicsSimulation::UpdatePostPhysics(
 	for (const FRigComponentKey& ComponentKey : SortedBodyComponentKeys)
 	{
 		FRigBodyRecord& Record = *BodyRecords.Find(ComponentKey);
+
 		if (Record.ActorHandle && Record.TargetElementKey.IsValid())
 		{
 			// Check the simulation output
@@ -1537,6 +1643,10 @@ void FRigPhysicsSimulation::UpdatePostPhysics(
 			// debugging (and this should be rare!)
 			Record.FinalComponentSpaceTM = ConvertSimSpaceTransformToComponentSpace(
 				SolverSettings, SimSpaceTM);
+
+			// Danny TODO rework this - currently in does the interpolation here, and stores it in
+			// the component, but I think the component should store the actual physics transform,
+			// and interpolation should only happen when we write back into the bone transform.
 			if (Alpha < 0.999f)
 			{
 				// Danny TODO Note that this uses Alpha to blend in component space.
@@ -1551,6 +1661,13 @@ void FRigPhysicsSimulation::UpdatePostPhysics(
 				Record.FinalComponentSpaceTM.SetRotation(TargetQ);
 				Record.FinalComponentSpaceTM.SetTranslation(TargetT);
 			}
+
+			if (FRigPhysicsBodyComponent* PhysicsBodyComponent = Cast<FRigPhysicsBodyComponent>(
+				Hierarchy.FindComponent(ComponentKey)))
+			{
+				PhysicsBodyComponent->ForceAndTorques.Empty();
+				SetPhysicsBodyComponentState(SolverSettings, Record, *PhysicsBodyComponent);
+			}
 		}
 	}
 
@@ -1579,13 +1696,13 @@ void FRigPhysicsSimulation::UpdatePostPhysics(
 			const FRigComponentKey& ComponentKey = BodyRecordPair.Key;
 			const FRigBodyRecord& Record = BodyRecordPair.Value;
 
-			if (const FRigPhysicsBodyComponent* PhysicsComponent = Cast<FRigPhysicsBodyComponent>(
+			if (const FRigPhysicsBodyComponent* PhysicsBodyComponent = Cast<FRigPhysicsBodyComponent>(
 				Hierarchy.FindComponent(ComponentKey)))
 			{
 				if (Record.ActorHandle && !Record.ActorHandle->GetIsKinematic())
 				{
 					// Get the TM in component space, and then convert it into sim space.
-					FRigElementKey SourceKey = PhysicsComponent->BodySolverSettings.SourceBone;
+					FRigElementKey SourceKey = PhysicsBodyComponent->BodySolverSettings.SourceBone;
 					if (!SourceKey.IsValid())
 					{
 						SourceKey = ComponentKey.ElementKey;
diff --git a/Engine/Plugins/Experimental/ControlRigPhysics/Source/ControlRigPhysics/Private/RigPhysicsSimulation.h b/Engine/Plugins/Experimental/ControlRigPhysics/Source/ControlRigPhysics/Private/RigPhysicsSimulation.h
--- a/Engine/Plugins/Experimental/ControlRigPhysics/Source/ControlRigPhysics/Private/RigPhysicsSimulation.h
+++ b/Engine/Plugins/Experimental/ControlRigPhysics/Source/ControlRigPhysics/Private/RigPhysicsSimulation.h
@@ -149,7 +149,7 @@ public:
 	// This will initialise/create the simulation and then create everything we need in it. 
 	void Instantiate(
 		const FRigVMExecuteContext&       ExecuteContext, 
-		const URigHierarchy&              Hierarchy, 
+		URigHierarchy&                    Hierarchy, 
 		const FRigPhysicsSolverComponent* SolverComponent);
 
 	// Integrates the simulation forwards.
@@ -221,7 +221,7 @@ private:
 
 	// Creates low level bodies
 	void InstantiatePhysicsBodies(
-		const URigHierarchy&              Hierarchy,
+		URigHierarchy&                    Hierarchy,
 		const FRigPhysicsSolverComponent* SolverComponent,
 		FRigPhysicsIgnorePairs&           IgnorePairs);
 
@@ -278,6 +278,12 @@ private:
 		const FRigBodyRecord&              Record, 
 		const FRigPhysicsBodyComponent*    PhysicsComponent);
 
+	void ApplyForcesPrePhysics(
+		const FRigVMExecuteContext&        ExecuteContext, 
+		const FRigPhysicsSolverComponent*  SolverComponent,
+		const FRigBodyRecord&              Record, 
+		const FRigPhysicsBodyComponent*    PhysicsComponent);
+
 	void UpdateJointPrePhysics(
 		const URigHierarchy&               Hierarchy, 
 		FRigJointRecord&                   JointRecord,
@@ -326,36 +332,57 @@ private:
 	// Converts a transform from component space (e.g. coming from the owning control rig) into the
 	// simulation space
 	FTransform ConvertComponentSpaceTransformToSimSpace(
-		const FRigPhysicsSolverSettings& SolverSettings, const FTransform& TM) const;
+		const FRigPhysicsSolverSettings& SolverSettings, const FTransform& ComponentSpaceTM) const;
+
+	// Converts a component space position to simulation space, including translation
+	FVector ConvertComponentSpacePositionToSimSpace(
+		const FRigPhysicsSolverSettings& SolverSettings, const FVector& ComponentSpacePosition) const;
 
+	// Converts a component space vector to simulation space, ignoring translation
 	FVector ConvertComponentSpaceVectorToSimSpace(
-		const FRigPhysicsSolverSettings& SolverSettings, const FVector& V) const;
+		const FRigPhysicsSolverSettings& SolverSettings, const FVector& ComponentSpaceVector) const;
+
+	// Converts a transform specified in world space to the simulation space
+	FTransform ConvertWorldTransformToSimSpace(
+		const FRigPhysicsSolverSettings& SolverSettings, const FTransform& WorldSpaceTM) const;
+
+	// Converts a position specified in world space to the simulation space
+	FVector ConvertWorldPositionToSimSpaceNoScale(
+		const FRigPhysicsSolverSettings& SolverSettings, const FVector& WorldSpacePosition) const;
+
+	// Converts a vector specified in world space to the simulation space (e.g. converting gravity)
+	FVector ConvertWorldVectorToSimSpaceNoScale(
+		const FRigPhysicsSolverSettings& SolverSettings, const FVector& WorldSpaceVector) const;
 
 	// Converts a transform from the simulation space to component space (e.g. for writing back to
 	// the owning control rig)
 	FTransform ConvertSimSpaceTransformToComponentSpace(
-		const FRigPhysicsSolverSettings& SolverSettings, const FTransform& TM) const;
+		const FRigPhysicsSolverSettings& SolverSettings, const FTransform& SimSpaceTM) const;
 
-	// Converts a vector specified in world space to the simulation space (e.g. converting gravity)
-	FVector ConvertWorldVectorToSimSpaceNoScale(
-		const FRigPhysicsSolverSettings& SolverSettings, const FVector& WorldVector) const;
+	// Converts a simulation space position to component space, including translation
+	FVector ConvertSimSpacePositionToComponentSpace(
+		const FRigPhysicsSolverSettings& SolverSettings, const FVector& SimSpacePosition) const;
 
-	// Converts a transform specified in world space to the simulation space
-	FTransform ConvertWorldTransformToSimSpace(
-		const FRigPhysicsSolverSettings& SolverSettings, const FTransform& WorldTM) const;
+	// Converts a simulation space position to component space, ignoring translation
+	FVector ConvertSimSpaceVectorToComponentSpace(
+		const FRigPhysicsSolverSettings& SolverSettings, const FVector& SimSpaceVector) const;
 
+	// Converts a collision space transform to simulation space
 	FTransform ConvertCollisionSpaceTransformToSimSpace(
 		const FRigPhysicsSolverSettings& SolverSettings, const FTransform& TM) const;
 
-	ImmediatePhysics::FJointHandle* CreateConstraint(
-		ImmediatePhysics::FActorHandle* ChildActorHandle, ImmediatePhysics::FActorHandle* ParentActorHandle);
-
 	// Gets the simulation actor handle for a component key. Note that the component key could be a
 	// body or a solver component. Can return nullptr
 	ImmediatePhysics::FActorHandle* GetActor(const FRigComponentKey& ComponentKey) const;
 
+	// Updates the PhysicsBodyComponent with the state of the rigid body from the record. Assumes
+	// that FinalComponentSpaceTM has been set.
+	void SetPhysicsBodyComponentState(
+		const FRigPhysicsSolverSettings& SolverSettings, 
+		FRigBodyRecord&                  Record, 
+		FRigPhysicsBodyComponent&        PhysicsBodyComponent) const;
+
 private:
-	TObjectPtr<UControlRig> OwningControlRig;
 	FRigComponentKey        PhysicsSolverComponentKey;
 
 	// A name that identifies what owns us - just used for debugging/logging
diff --git a/Engine/Plugins/Experimental/ControlRigPhysics/Source/ControlRigPhysics/Private/RigPhysicsSimulationSpace.cpp b/Engine/Plugins/Experimental/ControlRigPhysics/Source/ControlRigPhysics/Private/RigPhysicsSimulationSpace.cpp
--- a/Engine/Plugins/Experimental/ControlRigPhysics/Source/ControlRigPhysics/Private/RigPhysicsSimulationSpace.cpp
+++ b/Engine/Plugins/Experimental/ControlRigPhysics/Source/ControlRigPhysics/Private/RigPhysicsSimulationSpace.cpp
@@ -68,6 +68,19 @@ FTransform FRigPhysicsSimulation::ConvertComponentSpaceTransformToSimSpace(
 	}
 }
 
+//======================================================================================================================
+FVector FRigPhysicsSimulation::ConvertComponentSpacePositionToSimSpace(
+	const FRigPhysicsSolverSettings& SolverSettings, const FVector& P) const
+{
+	switch (SolverSettings.SimulationSpace)
+	{
+	case ERigPhysicsSimulationSpace::Component: return P;
+	case ERigPhysicsSimulationSpace::World:  return SimulationSpaceState.ComponentTM.TransformPositionNoScale(P);
+	case ERigPhysicsSimulationSpace::SpaceBone: return SimulationSpaceState.BoneRelComponentTM.InverseTransformPositionNoScale(P);
+	default: ensureMsgf(false, TEXT("Unsupported Simulation Space")); return P;
+	}
+}
+
 //======================================================================================================================
 FVector FRigPhysicsSimulation::ConvertComponentSpaceVectorToSimSpace(
 	const FRigPhysicsSolverSettings& SolverSettings, const FVector& V) const
@@ -75,8 +88,8 @@ FVector FRigPhysicsSimulation::ConvertComponentSpaceVectorToSimSpace(
 	switch (SolverSettings.SimulationSpace)
 	{
 	case ERigPhysicsSimulationSpace::Component: return V;
-	case ERigPhysicsSimulationSpace::World:  return SimulationSpaceState.ComponentTM.TransformVector(V);
-	case ERigPhysicsSimulationSpace::SpaceBone: return SimulationSpaceState.BoneRelComponentTM.InverseTransformVector(V);
+	case ERigPhysicsSimulationSpace::World:  return SimulationSpaceState.ComponentTM.TransformVectorNoScale(V);
+	case ERigPhysicsSimulationSpace::SpaceBone: return SimulationSpaceState.BoneRelComponentTM.InverseTransformVectorNoScale(V);
 	default: ensureMsgf(false, TEXT("Unsupported Simulation Space")); return V;
 	}
 }
@@ -94,6 +107,32 @@ FTransform FRigPhysicsSimulation::ConvertSimSpaceTransformToComponentSpace(
 	}
 }
 
+//======================================================================================================================
+FVector FRigPhysicsSimulation::ConvertSimSpacePositionToComponentSpace(
+	const FRigPhysicsSolverSettings& SolverSettings, const FVector& Position) const
+{
+	switch (SolverSettings.SimulationSpace)
+	{
+	case ERigPhysicsSimulationSpace::Component: return Position;
+	case ERigPhysicsSimulationSpace::World: return SimulationSpaceState.ComponentTM.InverseTransformPositionNoScale(Position);
+	case ERigPhysicsSimulationSpace::SpaceBone: return SimulationSpaceState.BoneRelComponentTM.TransformPositionNoScale(Position);
+	default: ensureMsgf(false, TEXT("Unsupported Simulation Space")); return Position;
+	}
+}
+
+//======================================================================================================================
+FVector FRigPhysicsSimulation::ConvertSimSpaceVectorToComponentSpace(
+	const FRigPhysicsSolverSettings& SolverSettings, const FVector& Vector) const
+{
+	switch (SolverSettings.SimulationSpace)
+	{
+	case ERigPhysicsSimulationSpace::Component: return Vector;
+	case ERigPhysicsSimulationSpace::World: return SimulationSpaceState.ComponentTM.InverseTransformVectorNoScale(Vector);
+	case ERigPhysicsSimulationSpace::SpaceBone: return SimulationSpaceState.BoneRelComponentTM.TransformVectorNoScale(Vector);
+	default: ensureMsgf(false, TEXT("Unsupported Simulation Space")); return Vector;
+	}
+}
+
 //======================================================================================================================
 FORCEINLINE_DEBUGGABLE FVector ConvertWorldVectorToSimSpaceNoScale(
 	ERigPhysicsSimulationSpace Space, const FVector& WorldVector, 
@@ -115,7 +154,8 @@ FVector FRigPhysicsSimulation::ConvertWorldVectorToSimSpaceNoScale(
 {
 	switch (SolverSettings.SimulationSpace)
 	{
-	case ERigPhysicsSimulationSpace::Component: return SimulationSpaceState.ComponentTM.InverseTransformVectorNoScale(WorldVector);
+	case ERigPhysicsSimulationSpace::Component: 
+		return SimulationSpaceState.ComponentTM.InverseTransformVectorNoScale(WorldVector);
 	case ERigPhysicsSimulationSpace::World: return WorldVector;
 	case ERigPhysicsSimulationSpace::SpaceBone:
 		return SimulationSpaceState.BoneRelComponentTM.InverseTransformVectorNoScale(
@@ -124,6 +164,22 @@ FVector FRigPhysicsSimulation::ConvertWorldVectorToSimSpaceNoScale(
 	}
 }
 
+//======================================================================================================================
+FVector FRigPhysicsSimulation::ConvertWorldPositionToSimSpaceNoScale(
+	const FRigPhysicsSolverSettings& SolverSettings, const FVector& WorldPosition) const
+{
+	switch (SolverSettings.SimulationSpace)
+	{
+	case ERigPhysicsSimulationSpace::Component: 
+		return SimulationSpaceState.ComponentTM.InverseTransformPositionNoScale(WorldPosition);
+	case ERigPhysicsSimulationSpace::World: return WorldPosition;
+	case ERigPhysicsSimulationSpace::SpaceBone:
+		return SimulationSpaceState.BoneRelComponentTM.InverseTransformPositionNoScale(
+			SimulationSpaceState.ComponentTM.InverseTransformPositionNoScale(WorldPosition));
+	default: ensureMsgf(false, TEXT("Unsupported Simulation Space")); return WorldPosition;
+	}
+}
+
 //======================================================================================================================
 FTransform FRigPhysicsSimulation::ConvertWorldTransformToSimSpace(
 	const FRigPhysicsSolverSettings& SolverSettings, const FTransform& WorldTM) const
diff --git a/Engine/Plugins/Experimental/ControlRigPhysics/Source/ControlRigPhysics/Public/RigPhysicsBodyComponent.h b/Engine/Plugins/Experimental/ControlRigPhysics/Source/ControlRigPhysics/Public/RigPhysicsBodyComponent.h
--- a/Engine/Plugins/Experimental/ControlRigPhysics/Source/ControlRigPhysics/Public/RigPhysicsBodyComponent.h
+++ b/Engine/Plugins/Experimental/ControlRigPhysics/Source/ControlRigPhysics/Public/RigPhysicsBodyComponent.h
@@ -3,7 +3,6 @@
 #pragma once
 
 #include "RigPhysicsData.h"
-
 #include "PhysicsControlData.h"
 
 #include "Rigs/RigPhysics.h"
@@ -64,6 +63,18 @@ public:
 	virtual FName GetDefaultComponentName() const override { return GetDefaultName(); }
 
 	static FName GetDefaultName() { return TEXT("PhysicsBody"); }
+
+	// Danny TODO Should any/all of these (a) be UPROPERTYs, and thus visible (b) result in a call
+	// to Hierarchy.Notify(...) when modified?
+
+	// Forces (etc) that should be applied and then cleared at the next solver step.
+	TArray<FPhysicsControlNamedForceAndTorqueData> ForceAndTorques;
+
+	// Internal/runtime state, in component space
+	FTransform Transform;
+	FTransform CoMTransform;
+	FVector    LinearVelocity = FVector::ZeroVector;
+	FVector    AngularVelocity = FVector::ZeroVector;
 };
 
 #undef UE_API
diff --git a/Engine/Plugins/Experimental/ControlRigPhysics/Source/ControlRigPhysics/Public/RigPhysicsBodyExecution.h b/Engine/Plugins/Experimental/ControlRigPhysics/Source/ControlRigPhysics/Public/RigPhysicsBodyExecution.h
--- a/Engine/Plugins/Experimental/ControlRigPhysics/Source/ControlRigPhysics/Public/RigPhysicsBodyExecution.h
+++ b/Engine/Plugins/Experimental/ControlRigPhysics/Source/ControlRigPhysics/Public/RigPhysicsBodyExecution.h
@@ -513,4 +513,202 @@ struct FRigUnit_HierarchySetPhysicsBodyDamping : public FRigUnit_PhysicsBaseMuta
 	float AngularDamping = 0;
 };
 
+/**
+ * This will add a force (or impulse, torque, velocity change etc) to the body. The force record 
+ * will be stored in the physics body component, and then applied next time the solver is stepped, 
+ * after which the record will be cleared.
+ */
+USTRUCT(meta = (DisplayName = "Add Force", Varying, Keywords = "Physics, Impulse, Torque, Acceleration, Velocity"))
+struct FRigUnit_HierarchyAddPhysicsBodyForce : public FRigUnit_PhysicsBaseMutable
+{
+	GENERATED_BODY()
+
+	FRigUnit_HierarchyAddPhysicsBodyForce()
+	{
+		PhysicsBodyComponentKey.ElementKey.Type = ERigElementType::Bone;
+		PhysicsBodyComponentKey.Name = FRigPhysicsBodyComponent::GetDefaultName();
+	}
+
+	RIGVM_METHOD()
+	UE_API virtual void Execute() override;
+
+	// The Physics Body to be updated
+	UPROPERTY(meta = (Input))
+	FRigComponentKey PhysicsBodyComponentKey;
+
+	// The desired force/torque to be applied
+	UPROPERTY(meta = (Input))
+	FPhysicsControlForceAndTorqueData ForceAndTorque;
+
+	// The optional name of this force record. Used to identify the force records, and potentially delete it.
+	// Can be blank.
+	UPROPERTY(meta = (Input))
+	FName Name;
+};
+
+/**
+ * Removes all force records (assuming they hasn't been applied yet) with the given Name on the Body component.
+ */
+USTRUCT(meta = (DisplayName = "Remove Force", Varying, Keywords = "Physics, Impulse, Torque, Acceleration, Velocity"))
+struct FRigUnit_HierarchyRemovePhysicsBodyForce : public FRigUnit_PhysicsBaseMutable
+{
+	GENERATED_BODY()
+
+	FRigUnit_HierarchyRemovePhysicsBodyForce()
+	{
+		PhysicsBodyComponentKey.ElementKey.Type = ERigElementType::Bone;
+		PhysicsBodyComponentKey.Name = FRigPhysicsBodyComponent::GetDefaultName();
+	}
+
+	RIGVM_METHOD()
+	UE_API virtual void Execute() override;
+
+	// The Physics Body to be updated
+	UPROPERTY(meta = (Input))
+	FRigComponentKey PhysicsBodyComponentKey;
+
+	// The name of the force record to remove.
+	UPROPERTY(meta = (Input))
+	FName Name;
+};
+
+/**
+ * Retrieves the current simulation transform of the body.
+ */
+USTRUCT(meta = (DisplayName = "Get Physics Body Transform", Varying))
+struct FRigUnit_HierarchyGetPhysicsBodyTransform : public FRigUnit_PhysicsBase
+{
+	GENERATED_BODY()
+
+	FRigUnit_HierarchyGetPhysicsBodyTransform()
+	{
+		PhysicsBodyComponentKey.ElementKey.Type = ERigElementType::Bone;
+		PhysicsBodyComponentKey.Name = FRigPhysicsBodyComponent::GetDefaultName();
+	}
+
+	RIGVM_METHOD()
+	UE_API virtual void Execute() override;
+
+	// The Physics Body to be queried
+	UPROPERTY(meta = (Input))
+	FRigComponentKey PhysicsBodyComponentKey;
+
+	// The current body transform.
+	UPROPERTY(meta = (Output))
+	FTransform Transform;
+};
+
+/**
+ * Retrieves the current simulation centre of mass transform of the body.
+ */
+USTRUCT(meta = (DisplayName = "Get Physics Body CoM Transform", Varying))
+struct FRigUnit_HierarchyGetPhysicsBodyCoMTransform : public FRigUnit_PhysicsBase
+{
+	GENERATED_BODY()
+
+	FRigUnit_HierarchyGetPhysicsBodyCoMTransform()
+	{
+		PhysicsBodyComponentKey.ElementKey.Type = ERigElementType::Bone;
+		PhysicsBodyComponentKey.Name = FRigPhysicsBodyComponent::GetDefaultName();
+	}
+
+	RIGVM_METHOD()
+	UE_API virtual void Execute() override;
+
+	// The Physics Body to be queried
+	UPROPERTY(meta = (Input))
+	FRigComponentKey PhysicsBodyComponentKey;
+
+	// The current centre of mass transform of the body.
+	UPROPERTY(meta = (Output))
+	FTransform CoMTransform;
+};
+
+/**
+ * Retrieves the current centre of mass linear velocity of the body.
+ */
+USTRUCT(meta = (DisplayName = "Get Physics Body Linear Velocity", Varying))
+struct FRigUnit_HierarchyGetPhysicsBodyLinearVelocity : public FRigUnit_PhysicsBase
+{
+	GENERATED_BODY()
+
+	FRigUnit_HierarchyGetPhysicsBodyLinearVelocity()
+	{
+		PhysicsBodyComponentKey.ElementKey.Type = ERigElementType::Bone;
+		PhysicsBodyComponentKey.Name = FRigPhysicsBodyComponent::GetDefaultName();
+	}
+
+	RIGVM_METHOD()
+	UE_API virtual void Execute() override;
+
+	// The Physics Body to be queried
+	UPROPERTY(meta = (Input))
+	FRigComponentKey PhysicsBodyComponentKey;
+
+	// The current body linear velocity.
+	UPROPERTY(meta = (Output))
+	FVector LinearVelocity = FVector::ZeroVector;
+};
+
+/**
+ * Retrieves the current centre of mass angular velocity of the body.
+ */
+USTRUCT(meta = (DisplayName = "Get Physics Body Angular Velocity", Varying))
+struct FRigUnit_HierarchyGetPhysicsBodyAngularVelocity : public FRigUnit_PhysicsBase
+{
+	GENERATED_BODY()
+
+	FRigUnit_HierarchyGetPhysicsBodyAngularVelocity()
+	{
+		PhysicsBodyComponentKey.ElementKey.Type = ERigElementType::Bone;
+		PhysicsBodyComponentKey.Name = FRigPhysicsBodyComponent::GetDefaultName();
+	}
+
+	RIGVM_METHOD()
+	UE_API virtual void Execute() override;
+
+	// The Physics Body to be queried
+	UPROPERTY(meta = (Input))
+	FRigComponentKey PhysicsBodyComponentKey;
+
+	// The current body angular velocity.
+	UPROPERTY(meta = (Output))
+	FVector AngularVelocity = FVector::ZeroVector;
+};
+
+/**
+ * Retrieves the current linear velocity of the body at a point in space.
+ */
+USTRUCT(meta = (DisplayName = "Get Physics Body Point Velocity", Varying))
+struct FRigUnit_HierarchyGetPhysicsBodyPointVelocity : public FRigUnit_PhysicsBase
+{
+	GENERATED_BODY()
+
+	FRigUnit_HierarchyGetPhysicsBodyPointVelocity()
+	{
+		PhysicsBodyComponentKey.ElementKey.Type = ERigElementType::Bone;
+		PhysicsBodyComponentKey.Name = FRigPhysicsBodyComponent::GetDefaultName();
+	}
+
+	RIGVM_METHOD()
+	UE_API virtual void Execute() override;
+
+	// The Physics Body to be queried
+	UPROPERTY(meta = (Input))
+	FRigComponentKey PhysicsBodyComponentKey;
+
+	// The position to calculate the velocity at.
+	UPROPERTY(meta = (Input))
+	FVector Position;
+
+	// The space that Position is specified in.
+	UPROPERTY(meta = (Input))
+	EPhysicsControlSpace Space = EPhysicsControlSpace::Body;
+
+	// The current velocity at Position.
+	UPROPERTY(meta = (Output))
+	FVector Velocity = FVector::ZeroVector;
+};
+
+
 #undef UE_API
diff --git a/Engine/Plugins/Experimental/ControlRigPhysics/Source/ControlRigPhysics/Public/RigPhysicsData.h b/Engine/Plugins/Experimental/ControlRigPhysics/Source/ControlRigPhysics/Public/RigPhysicsData.h
--- a/Engine/Plugins/Experimental/ControlRigPhysics/Source/ControlRigPhysics/Public/RigPhysicsData.h
+++ b/Engine/Plugins/Experimental/ControlRigPhysics/Source/ControlRigPhysics/Public/RigPhysicsData.h
@@ -424,10 +424,10 @@ struct FRigPhysicsSolverSettings
 
 	friend FArchive& operator <<(FArchive& Ar, FRigPhysicsSolverSettings& Data);
 
-	// If true, then any physics component below (the element owning this one) will be automatically
-	// added to this solver, if its UseAutomaticSolver flag is set.
+	// If true, then any physics body component below (the element owning this one) will be automatically
+	// added to this solver, if its UseAutomaticSolver/AutoFindSolver flag is set.
 	UPROPERTY(EditAnywhere, Category = SolverSettings)
-	bool bAutomaticallyAddPhysicsComponents = true;
+	bool bAutomaticallyAddPhysicsBodyComponents = true;
 
 	UPROPERTY(BlueprintReadOnly, EditAnywhere, Category = SolverSettings)
 	ERigPhysicsSimulationSpace SimulationSpace = ERigPhysicsSimulationSpace::Component;
@@ -758,7 +758,7 @@ struct FRigPhysicsBodySolverSettings
 	// If true (and the physics solver is not explicitly set), then this component will be added to
 	// any physics solver that exists above it in the hierarchy, if that solver allows automatically
 	// adding physics components.
-	UPROPERTY(EditAnywhere, Category = Solver)
+	UPROPERTY(EditAnywhere, DisplayName="Auto Find Solver", Category = Solver)
 	bool bUseAutomaticSolver = true;
 
 	// The bone that is used to initialize physics, as well as what to track when the body is set to be kinematic.
diff --git a/Engine/Plugins/Experimental/PhysicsControl/Source/PhysicsControl/Private/AnimNode_RigidBodyWithControl.cpp b/Engine/Plugins/Experimental/PhysicsControl/Source/PhysicsControl/Private/AnimNode_RigidBodyWithControl.cpp
--- a/Engine/Plugins/Experimental/PhysicsControl/Source/PhysicsControl/Private/AnimNode_RigidBodyWithControl.cpp
+++ b/Engine/Plugins/Experimental/PhysicsControl/Source/PhysicsControl/Private/AnimNode_RigidBodyWithControl.cpp
@@ -336,21 +336,18 @@ void FAnimNode_RigidBodyWithControl::InitSimulationSpace(
 	const FTransform& ComponentToWorld,
 	const FTransform& BoneToComponent)
 {
-	FTransform SimulationSpaceTransform = SpaceToWorldTransform(SimulationSpace, ComponentToWorld, BoneToComponent);
+	SimulationSpaceTransform = SpaceToWorldTransform(SimulationSpace, ComponentToWorld, BoneToComponent);
 	PreviousSimulationSpaceTransform = SimulationSpaceTransform;
 	PreviousPreviousSimulationSpaceTransform = SimulationSpaceTransform;
 }
 
 void FAnimNode_RigidBodyWithControl::CalculateSimulationSpaceMotion(
-	ESimulationSpace         Space, 
-	const FTransform& SimulationSpaceTransform,
-	const FTransform&        ComponentToWorld, 
 	const float              Dt,
 	const FSimSpaceSettings& Settings,
-	FVector& SpaceLinearVel, 
-	FVector& SpaceAngularVel, 
-	FVector& SpaceLinearAcc, 
-	FVector& SpaceAngularAcc)
+	FVector&                 SpaceLinearVel, 
+	FVector&                 SpaceAngularVel, 
+	FVector&                 SpaceLinearAcc, 
+	FVector&                 SpaceAngularAcc)
 {
 	SpaceLinearVel = FVector::ZeroVector;
 	SpaceAngularVel = FVector::ZeroVector;
@@ -364,7 +361,7 @@ void FAnimNode_RigidBodyWithControl::CalculateSimulationSpaceMotion(
 	}
 
 	// If the simulation is in world space, the simulation space is stationary
-	if (Space == ESimulationSpace::WorldSpace)
+	if (SimulationSpace == ESimulationSpace::WorldSpace)
 	{
 		SpaceLinearVel = Settings.ExternalLinearVelocity;
 		SpaceAngularVel = Settings.ExternalAngularVelocity;
@@ -372,13 +369,17 @@ void FAnimNode_RigidBodyWithControl::CalculateSimulationSpaceMotion(
 	}
 
 	// World-space component linear velocity and acceleration
-	FVector PrevSpaceLinearVel = Chaos::FVec3::CalculateVelocity(PreviousPreviousSimulationSpaceTransform.GetTranslation(), PreviousSimulationSpaceTransform.GetTranslation(), Dt);
-	SpaceLinearVel = Chaos::FVec3::CalculateVelocity(PreviousSimulationSpaceTransform.GetTranslation(), SimulationSpaceTransform.GetTranslation(), Dt);
+	FVector PrevSpaceLinearVel = Chaos::FVec3::CalculateVelocity(
+		PreviousPreviousSimulationSpaceTransform.GetTranslation(), PreviousSimulationSpaceTransform.GetTranslation(), Dt);
+	SpaceLinearVel = Chaos::FVec3::CalculateVelocity(
+		PreviousSimulationSpaceTransform.GetTranslation(), SimulationSpaceTransform.GetTranslation(), Dt);
 	SpaceLinearAcc = (SpaceLinearVel - PrevSpaceLinearVel) / Dt;
 
 	// World-space component angular velocity and acceleration
-	FVector PrevSpaceAngularVel = Chaos::FRotation3::CalculateAngularVelocity(PreviousPreviousSimulationSpaceTransform.GetRotation(), PreviousSimulationSpaceTransform.GetRotation(), Dt);
-	SpaceAngularVel = Chaos::FRotation3::CalculateAngularVelocity(PreviousSimulationSpaceTransform.GetRotation(), SimulationSpaceTransform.GetRotation(), Dt);
+	FVector PrevSpaceAngularVel = Chaos::FRotation3::CalculateAngularVelocity(
+		PreviousPreviousSimulationSpaceTransform.GetRotation(), PreviousSimulationSpaceTransform.GetRotation(), Dt);
+	SpaceAngularVel = Chaos::FRotation3::CalculateAngularVelocity(
+		PreviousSimulationSpaceTransform.GetRotation(), SimulationSpaceTransform.GetRotation(), Dt);
 	SpaceAngularAcc = (SpaceAngularVel - PrevSpaceAngularVel) / Dt;
 
 	// Apply Z scale
@@ -545,6 +546,9 @@ void FAnimNode_RigidBodyWithControl::EvaluateSkeletalControl_AnyThread(FComponen
 
 		const FTransform BaseBoneTM = Output.Pose.GetComponentSpaceTransform(BaseBoneRef.GetCompactPoseIndex(BoneContainer));
 
+		// Now we know the base bone, we can calculate the simulation space transform
+		SimulationSpaceTransform = SpaceToWorldTransform(SimulationSpace, CompWorldSpaceTM, BaseBoneTM);
+
 		// Initialize potential new bodies because of LOD change.
 		if (ResetSimulatedTeleportType == ETeleportType::None && bCheckForBodyTransformInit)
 		{
@@ -761,6 +765,7 @@ void FAnimNode_RigidBodyWithControl::EvaluateSkeletalControl_AnyThread(FComponen
 				// can apply to bodies that aren't under the influence of a modifier
 				ApplyControlAndModifierUpdatesAndParametersToRecords(ControlAndModifierUpdates, ControlAndModifierParameters);
 				ApplyControlsAndModifiers(SimSpaceGravity, DeltaSeconds);
+				ApplyForcesAndTorques(CompWorldSpaceTM, BaseBoneTM);
 			}
 
 			// Note that the simulation interpolates kinematic targets to handle substepping
@@ -789,15 +794,11 @@ void FAnimNode_RigidBodyWithControl::EvaluateSkeletalControl_AnyThread(FComponen
 				SimSpaceSettingsToUse = &RBANWithControl_SimSpaceOverride;
 			}
 
-			const FTransform SpaceTransform = SpaceToWorldTransform(SimulationSpace, CompWorldSpaceTM, BaseBoneTM);
 			FVector SimulationSpaceLinearVelocity;
 			FVector SimulationSpaceAngularVelocity;
 			FVector SimulationSpaceLinearAcceleration;
 			FVector SimulationSpaceAngularAcceleration;
 			CalculateSimulationSpaceMotion(
-				SimulationSpace,
-				SpaceTransform,
-				CompWorldSpaceTM,
 				DeltaSeconds,
 				*SimSpaceSettingsToUse,
 				SimulationSpaceLinearVelocity,
@@ -805,11 +806,11 @@ void FAnimNode_RigidBodyWithControl::EvaluateSkeletalControl_AnyThread(FComponen
 				SimulationSpaceLinearAcceleration,
 				SimulationSpaceAngularAcceleration);
 
-			UpdateWorldObjects(SpaceTransform, SimulationSpaceLinearVelocity, SimulationSpaceAngularVelocity, DeltaSeconds);
-			UpdateClothColliderObjects(SpaceTransform);
+			UpdateWorldObjects(SimulationSpaceLinearVelocity, SimulationSpaceAngularVelocity, DeltaSeconds);
+			UpdateClothColliderObjects();
 
 			Simulation->UpdateSimulationSpace(
-				SpaceTransform, SimulationSpaceLinearVelocity, SimulationSpaceAngularVelocity,
+				SimulationSpaceTransform, SimulationSpaceLinearVelocity, SimulationSpaceAngularVelocity,
 				SimulationSpaceLinearAcceleration, SimulationSpaceAngularAcceleration);
 
 			Simulation->SetSimulationSpaceSettings(
@@ -1606,7 +1607,7 @@ void FAnimNode_RigidBodyWithControl::RemoveClothColliderObjects()
 	ClothColliders.Reset();
 }
 
-void FAnimNode_RigidBodyWithControl::UpdateClothColliderObjects(const FTransform& SpaceTransform)
+void FAnimNode_RigidBodyWithControl::UpdateClothColliderObjects()
 {
 	for (FClothCollider& ClothCollider : ClothColliders)
 	{
@@ -1615,9 +1616,9 @@ void FAnimNode_RigidBodyWithControl::UpdateClothColliderObjects(const FTransform
 			// Calculate the sim-space transform of this object
 			const FTransform CompWorldTransform = ClothCollider.SkeletalMeshComponent->GetBoneTransform(ClothCollider.BoneIndex);
 			FTransform CompSpaceTransform;
-			CompSpaceTransform.SetTranslation(SpaceTransform.InverseTransformPosition(
+			CompSpaceTransform.SetTranslation(SimulationSpaceTransform.InverseTransformPosition(
 				CompWorldTransform.GetLocation()));
-			CompSpaceTransform.SetRotation(SpaceTransform.InverseTransformRotation(
+			CompSpaceTransform.SetRotation(SimulationSpaceTransform.InverseTransformRotation(
 				CompWorldTransform.GetRotation()));
 			CompSpaceTransform.SetScale3D(FVector::OneVector);	// TODO - sort out scale for world objects in local sim
 
@@ -1762,7 +1763,7 @@ void FAnimNode_RigidBodyWithControl::PurgeExpiredWorldObjects()
 // component-based simulation and need to update stationary objects into our space, or simply
 // because the objects are moving.
 void FAnimNode_RigidBodyWithControl::UpdateWorldObjects(
-	const FTransform& SpaceTransform, const FVector& LinearVel, const FVector& AngularVel, const float DeltaSeconds)
+	const FVector& LinearVel, const FVector& AngularVel, const float DeltaSeconds)
 {
 	SCOPE_CYCLE_COUNTER(STAT_RigidBodyNodeWithControl_UpdateWorldObjects);
 	LLM_SCOPE_BYNAME(TEXT("Animation/RigidBodyWithControl"));
@@ -1777,10 +1778,10 @@ void FAnimNode_RigidBodyWithControl::UpdateWorldObjects(
 		// We will want the "previous" simulation space to calculate velocities. Get it from the
 		// velocity, so we can take advantage of that being zero at initialisation/reset times etc.
 		const FVector PrevSimSpacePosition = 
-			SpaceTransform.GetLocation() - LinearVel * DeltaSeconds;
+			SimulationSpaceTransform.GetLocation() - LinearVel * DeltaSeconds;
 		const FQuat SimSpaceDeltaOrientation = FQuat::MakeFromRotationVector(AngularVel * DeltaSeconds);
 		const FQuat PrevSimSpaceOrientation = 
-			SimSpaceDeltaOrientation.Inverse() * SpaceTransform.GetRotation();
+			SimSpaceDeltaOrientation.Inverse() * SimulationSpaceTransform.GetRotation();
 		const FTransform PrevSimSpaceTransform(PrevSimSpaceOrientation, PrevSimSpacePosition, FVector::OneVector);
 
 		for (const auto& WorldEntry : ComponentsInSim)
@@ -1792,14 +1793,14 @@ void FAnimNode_RigidBodyWithControl::UpdateWorldObjects(
 
 				// Calculate the sim-space transform of this object
 				const FTransform CompWorldTransform = OverlapComp->BodyInstance.GetUnrealWorldTransform();
-				const FVector WorldPosition = CompWorldTransform.GetLocation();
-				const FQuat WorldOrientation = CompWorldTransform.GetRotation();
+				const FVector CompWorldPosition = CompWorldTransform.GetLocation();
+				const FQuat CompWorldOrientation = CompWorldTransform.GetRotation();
 
-				const FVector CompSpacePosition = SpaceTransform.InverseTransformPosition(WorldPosition);
-				const FQuat CompSpaceOrientation = SpaceTransform.InverseTransformRotation(WorldOrientation);
+				const FVector CompSimSpacePosition = SimulationSpaceTransform.InverseTransformPosition(CompWorldPosition);
+				const FQuat CompSimSpaceOrientation = SimulationSpaceTransform.InverseTransformRotation(CompWorldOrientation);
 
 				// TODO - sort out scale for world objects in local sim
-				const FTransform CompSpaceTransform(CompSpaceOrientation, CompSpacePosition, FVector::OneVector);
+				const FTransform CompSimSpaceTransform(CompSimSpaceOrientation, CompSimSpacePosition, FVector::OneVector);
 
 				// We need to set the velocity - either because the object or the space is
 				// moving. Ideally would do this by setting the velocity (with transformations), but
@@ -1815,16 +1816,16 @@ void FAnimNode_RigidBodyWithControl::UpdateWorldObjects(
 				// To get the previous component space transform, we estimate the previous world space
 				// transform, and then put that into the previous sim space.
 
-				const FVector PrevWorldPosition = WorldPosition - Velocity * DeltaSeconds;
+				const FVector PrevCompWorldPosition = CompWorldPosition - Velocity * DeltaSeconds;
 				const FQuat DeltaWorldOrientation = FQuat::MakeFromRotationVector(AngularVelocity * DeltaSeconds);
-				const FQuat PrevWorldOrientation = DeltaWorldOrientation.Inverse() * WorldOrientation;
+				const FQuat PrevCompWorldOrientation = DeltaWorldOrientation.Inverse() * CompWorldOrientation;
 
-				const FVector PrevCompSpacePosition = PrevSimSpaceTransform.InverseTransformPosition(WorldPosition);
-				const FQuat PrevCompSpaceOrientation = PrevSimSpaceTransform.InverseTransformRotation(WorldOrientation);
-				const FTransform PrevCompSpaceTransform(PrevCompSpaceOrientation, PrevCompSpacePosition, FVector::OneVector);
+				const FVector PrevCompSimSpacePosition = PrevSimSpaceTransform.InverseTransformPosition(CompWorldPosition);
+				const FQuat PrevCompSimSpaceOrientation = PrevSimSpaceTransform.InverseTransformRotation(CompWorldOrientation);
+				const FTransform PrevCompSimSpaceTransform(PrevCompSimSpaceOrientation, PrevCompSimSpacePosition, FVector::OneVector);
 
-				ActorHandle->InitWorldTransform(PrevCompSpaceTransform);
-				ActorHandle->SetKinematicTarget(CompSpaceTransform);
+				ActorHandle->InitWorldTransform(PrevCompSimSpaceTransform);
+				ActorHandle->SetKinematicTarget(CompSimSpaceTransform);
 			}
 		}
 	}
@@ -1839,16 +1840,16 @@ void FAnimNode_RigidBodyWithControl::UpdateWorldObjects(
 
 				// Calculate the sim-space transform of this object
 				const FTransform CompWorldTransform = OverlapComp->BodyInstance.GetUnrealWorldTransform();
-				const FVector WorldPosition = CompWorldTransform.GetLocation();
-				const FQuat WorldOrientation = CompWorldTransform.GetRotation();
+				const FVector CompWorldPosition = CompWorldTransform.GetLocation();
+				const FQuat CompWorldOrientation = CompWorldTransform.GetRotation();
 
-				const FVector CompSpacePosition = SpaceTransform.InverseTransformPosition(WorldPosition);
-				const FQuat CompSpaceOrientation = SpaceTransform.InverseTransformRotation(WorldOrientation);
+				const FVector CompSimSpacePosition = SimulationSpaceTransform.InverseTransformPosition(CompWorldPosition);
+				const FQuat CompSimSpaceOrientation = SimulationSpaceTransform.InverseTransformRotation(CompWorldOrientation);
 
 				// TODO - sort out scale for world objects in local sim
-				const FTransform CompSpaceTransform(CompSpaceOrientation, CompSpacePosition, FVector::OneVector);
+				const FTransform CompSimSpaceTransform(CompSimSpaceOrientation, CompSimSpacePosition, FVector::OneVector);
 
-				ActorHandle->InitWorldTransform(CompSpaceTransform);
+				ActorHandle->InitWorldTransform(CompSimSpaceTransform);
 			}
 		}
 	}
diff --git a/Engine/Plugins/Experimental/PhysicsControl/Source/PhysicsControl/Private/PhysicsControlData.cpp b/Engine/Plugins/Experimental/PhysicsControl/Source/PhysicsControl/Private/PhysicsControlData.cpp
--- a/Engine/Plugins/Experimental/PhysicsControl/Source/PhysicsControl/Private/PhysicsControlData.cpp
+++ b/Engine/Plugins/Experimental/PhysicsControl/Source/PhysicsControl/Private/PhysicsControlData.cpp
@@ -12,6 +12,34 @@
 #define SELECT_PARAM(NAME) Output.NAME = (Weight < 0.5f) ? A.NAME : B.NAME
 #define SET_ENABLED_PARAM(NAME) Output.bEnable##NAME = A.bEnable##NAME && B.bEnable##NAME
 
+//======================================================================================================================
+EPhysicsControlForceType ConvertForceType(ImmediatePhysics_Shared::EForceType InForceType)
+{
+	switch (InForceType)
+	{
+	case ImmediatePhysics_Shared::EForceType::AddForce: return EPhysicsControlForceType::Force;
+	case ImmediatePhysics_Shared::EForceType::AddAcceleration: return EPhysicsControlForceType::Acceleration;
+	case ImmediatePhysics_Shared::EForceType::AddImpulse: return EPhysicsControlForceType::Impulse;
+	case ImmediatePhysics_Shared::EForceType::AddVelocity: return EPhysicsControlForceType::VelocityChange;
+	default:
+		return EPhysicsControlForceType::Force;
+	}
+}
+
+//======================================================================================================================
+ImmediatePhysics_Shared::EForceType ConvertForceType(EPhysicsControlForceType InForceType)
+{
+	switch (InForceType)
+	{
+	case EPhysicsControlForceType::Force: return ImmediatePhysics_Shared::EForceType::AddForce;
+	case EPhysicsControlForceType::Acceleration: return ImmediatePhysics_Shared::EForceType::AddAcceleration;
+	case EPhysicsControlForceType::Impulse: return ImmediatePhysics_Shared::EForceType::AddImpulse;
+	case EPhysicsControlForceType::VelocityChange: return ImmediatePhysics_Shared::EForceType::AddVelocity;
+	default:
+		return ImmediatePhysics_Shared::EForceType::AddForce;
+	}
+}
+
 //======================================================================================================================
 FPhysicsControlSetUpdates& FPhysicsControlSetUpdates::operator+=(const FPhysicsControlSetUpdates& Other)
 {
diff --git a/Engine/Plugins/Experimental/PhysicsControl/Source/PhysicsControl/Private/RigidBodyWithControl.cpp b/Engine/Plugins/Experimental/PhysicsControl/Source/PhysicsControl/Private/RigidBodyWithControl.cpp
--- a/Engine/Plugins/Experimental/PhysicsControl/Source/PhysicsControl/Private/RigidBodyWithControl.cpp
+++ b/Engine/Plugins/Experimental/PhysicsControl/Source/PhysicsControl/Private/RigidBodyWithControl.cpp
@@ -27,11 +27,10 @@ DECLARE_CYCLE_STAT(TEXT("RigidBodyNodeWithControl_LogControlsModifiersAndSets"),
 DECLARE_CYCLE_STAT(TEXT("RigidBodyNodeWithControl_ApplyControlAndModifierUpdatesAndParametersToRecords"), STAT_RigidBodyNodeWithControl_ApplyControlAndModifierUpdatesAndParametersToRecords, STATGROUP_Anim);
 DECLARE_CYCLE_STAT(TEXT("RigidBodyNodeWithControl_ApplyControlsAndModifierDatas"), STAT_RigidBodyNodeWithControl_ApplyControlsAndModifierDatas, STATGROUP_Anim);
 DECLARE_CYCLE_STAT(TEXT("RigidBodyNodeWithControl_ApplyControlsAndModifiers"), STAT_RigidBodyNodeWithControl_ApplyControlsAndModifiers, STATGROUP_Anim);
-
+DECLARE_CYCLE_STAT(TEXT("RigidBodyNodeWithControl_ApplyForcesAndTorques"), STAT_RigidBodyNodeWithControl_ApplyForcesAndTorques, STATGROUP_Anim);
 DECLARE_CYCLE_STAT(TEXT("RigidBodyNodeWithControl_ApplyControls"), STAT_RigidBodyNodeWithControl_ApplyControls, STATGROUP_Anim);
 DECLARE_CYCLE_STAT(TEXT("RigidBodyNodeWithControl_ApplyModifiers"), STAT_RigidBodyNodeWithControl_ApplyModifiers, STATGROUP_Anim);
 DECLARE_CYCLE_STAT(TEXT("RigidBodyNodeWithControl_SetPhysicsJointEnabled"), STAT_RigidBodyNodeWithControl_SetPhysicsJointEnabled, STATGROUP_Anim);
-
 DECLARE_CYCLE_STAT(TEXT("RigidBodyNodeWithControl_ApplyKinematicTargets"), STAT_RigidBodyNodeWithControl_ApplyKinematicTargets, STATGROUP_Anim);
 DECLARE_CYCLE_STAT(TEXT("RigidBodyNodeWithControl_ApplyCurrentConstraintProfile"), STAT_RigidBodyNodeWithControl_ApplyCurrentConstraintProfile, STATGROUP_Anim);
 
@@ -919,6 +918,109 @@ void FAnimNode_RigidBodyWithControl::ApplyControlsAndModifiers(const FVector& Si
 	}
 }
 
+//======================================================================================================================
+void FAnimNode_RigidBodyWithControl::ApplyForceAndTorque(
+	const FPhysicsControlForceAndTorqueData& ForceAndTorqueData,
+	const FRigidBodyModifierRecord&          ModifierRecord,
+	const FTransform&                        ComponentToWorld, 
+	const FTransform&                        BaseBoneTM)
+{
+	if (ModifierRecord.ActorHandle && ModifierRecord.ActorHandle->IsSimulated())
+	{
+		ImmediatePhysics::FActorHandle* ActorHandle = ModifierRecord.ActorHandle;
+
+		// Need to convert the requested force/torque into simulation space
+		FVector LinearSS;
+		FVector AngularSS;
+
+		switch (ForceAndTorqueData.ForceAndTorqueSpace)
+		{
+		case EPhysicsControlSpace::World:
+			LinearSS = WorldVectorToSpaceNoScale(
+				SimulationSpace, ForceAndTorqueData.Linear, ComponentToWorld, BaseBoneTM);
+			AngularSS = WorldVectorToSpaceNoScale(
+				SimulationSpace, ForceAndTorqueData.Angular, ComponentToWorld, BaseBoneTM);
+			break;
+		case EPhysicsControlSpace::Component:
+			LinearSS = ConvertCSVectorToSimSpace(
+				SimulationSpace, ForceAndTorqueData.Linear, ComponentToWorld, BaseBoneTM);
+			AngularSS = ConvertCSVectorToSimSpace
+				(SimulationSpace, ForceAndTorqueData.Angular, ComponentToWorld, BaseBoneTM);
+			break;
+		case EPhysicsControlSpace::Body:
+			LinearSS = ActorHandle->GetWorldTransform().GetRotation() * ForceAndTorqueData.Linear;
+			AngularSS = ActorHandle->GetWorldTransform().GetRotation() * ForceAndTorqueData.Angular;
+			break;
+		default:
+			LinearSS = FVector::ZeroVector;
+			AngularSS = FVector::ZeroVector;
+			break;
+		}
+
+		// Apply it differently depending on the location
+		ImmediatePhysics_Shared::EForceType FT = ConvertForceType(ForceAndTorqueData.Type);
+		switch (ForceAndTorqueData.LocationSpace)
+		{
+		case EPhysicsControlSpace::World:
+		{
+			FVector SimSpaceLocation = WorldPositionToSpace(
+				SimulationSpace, ForceAndTorqueData.Location, ComponentToWorld, BaseBoneTM);
+			ActorHandle->AddForceAtLocation(LinearSS, SimSpaceLocation, FT);
+		}
+		break;
+		case EPhysicsControlSpace::Component:
+		{
+			FVector SimSpaceLocation = ConvertCSPositionToSimSpace(
+				SimulationSpace, ForceAndTorqueData.Location, ComponentToWorld, BaseBoneTM);
+			ActorHandle->AddForceAtLocation(LinearSS, SimSpaceLocation, FT);
+		}
+		break;
+		case EPhysicsControlSpace::Body:
+		{
+			FVector SimSpaceLocation = ActorHandle->GetWorldTransform().TransformPositionNoScale(
+				ActorHandle->GetLocalCoMLocation() + ForceAndTorqueData.Location);
+			ActorHandle->AddForceAtLocation(LinearSS, SimSpaceLocation, FT);
+		}
+		break;
+		default:
+			break;
+		}
+		// Torque doesn't get applied at a location
+		ActorHandle->AddTorque(AngularSS, FT);
+	}
+}
+
+//======================================================================================================================
+void FAnimNode_RigidBodyWithControl::ApplyForcesAndTorques(
+	const FTransform& ComponentToWorld, const FTransform& BaseBoneTM)
+{
+	SCOPE_CYCLE_COUNTER(STAT_RigidBodyNodeWithControl_ApplyForcesAndTorques);
+	for (TMap<FName, FPhysicsControlForceAndTorqueData>::ElementType& ForceAndTorquePair : ForcesAndTorques.ForcesAndTorques)
+	{
+		FName Name = ForceAndTorquePair.Key;
+		const FPhysicsControlForceAndTorqueData& ForceAndTorqueData = ForceAndTorquePair.Value;
+
+		const TArray<FName>& Names = NameRecords.GetBodyModifierNamesInSet(Name);
+		if (Names.IsEmpty())
+		{
+			if (FRigidBodyModifierRecord* ModifierRecord = ModifierRecords.Find(Name))
+			{
+				ApplyForceAndTorque(ForceAndTorqueData, *ModifierRecord, ComponentToWorld, BaseBoneTM);
+			}
+		}
+		else
+		{
+			for (const FName& ModifierName : Names)
+			{
+				if (FRigidBodyModifierRecord* ModifierRecord = ModifierRecords.Find(ModifierName))
+				{
+					ApplyForceAndTorque(ForceAndTorqueData, *ModifierRecord, ComponentToWorld, BaseBoneTM);
+				}
+			}
+		}
+	}
+}
+
 //======================================================================================================================
 // Note that this will be called AFTER normal kinematic targets have been set
 void FAnimNode_RigidBodyWithControl::ApplyKinematicTargets()
diff --git a/Engine/Plugins/Experimental/PhysicsControl/Source/PhysicsControl/Public/AnimNode_RigidBodyWithControl.h b/Engine/Plugins/Experimental/PhysicsControl/Source/PhysicsControl/Public/AnimNode_RigidBodyWithControl.h
--- a/Engine/Plugins/Experimental/PhysicsControl/Source/PhysicsControl/Public/AnimNode_RigidBodyWithControl.h
+++ b/Engine/Plugins/Experimental/PhysicsControl/Source/PhysicsControl/Public/AnimNode_RigidBodyWithControl.h
@@ -389,6 +389,12 @@ public:
 	UPROPERTY(EditAnywhere, Category = Controls, meta = (PinShownByDefault))
 	FRigidBodyKinematicTargets KinematicTargets;
 
+	/**
+	 * Forces/torques that should be applied to dynamic bodies that are under the influence of a body modifier.
+	 */
+	UPROPERTY(EditAnywhere, Category = Controls, meta = (PinShownByDefault))
+	FRigidBodyForcesAndTorques ForcesAndTorques;
+
 	/** 
 	 * These options configure the way constraint transforms can be modified at runtime to match the current Skeletal
 	 * Mesh. This can be useful when applying a single RBwC node to assets with different skeletons.
@@ -450,7 +456,7 @@ private:
 
 	UE_API void InitPhysics(const UAnimInstance* InAnimInstance);
 	UE_API void UpdateWorldGeometry(const UWorld& World, const USkeletalMeshComponent& SKC);
-	UE_API void UpdateWorldForces(const FTransform& ComponentToWorld, const FTransform& RootBoneTM, const float DeltaSeconds);
+	UE_API void UpdateWorldForces(const FTransform& ComponentToWorld, const FTransform& BaseBoneTM, const float DeltaSeconds);
 
 public:
 	UE_API const FTransform GetBodyTransform(const int32 BodyIndex) const;
@@ -463,7 +469,7 @@ public:
 		const FName ModifierName, const FName BoneName, const FPhysicsControlModifierData& ModifierData);
 	UE_API FName CreateBodyModifier(const FName BoneName, const FPhysicsControlModifierData& ModifierData);
 
-private:
+protected:
 
 	// Add a new physics body and update the body name to index map
 	UE_API int32 AddBody(ImmediatePhysics::FActorHandle* const BodyHandle); 
@@ -480,6 +486,15 @@ private:
 
 	UE_API void ApplyControlsAndModifiers(const FVector& SimSpaceGravity, float DeltaTime);
 
+	UE_API void ApplyForcesAndTorques(
+		const FTransform& ComponentToWorld, const FTransform& BaseBoneTM);
+
+	UE_API void ApplyForceAndTorque(
+		const FPhysicsControlForceAndTorqueData& ForceAndTorqueData,
+		const FRigidBodyModifierRecord& ModifierRecord,
+		const FTransform& ComponentToWorld,
+		const FTransform& BaseBoneTM);
+
 	// This applies kinematic targets to any that have been set through the body modifiers
 	UE_API void ApplyKinematicTargets();
 
@@ -523,9 +538,6 @@ private:
 
 	// Calculate simulation space transform, velocity etc to pass into the solver
 	UE_API void CalculateSimulationSpaceMotion(
-		ESimulationSpace Space,
-		const FTransform& SpaceTransform,
-		const FTransform& ComponentToWorld,
 		const float Dt,
 		const FSimSpaceSettings& Settings,
 		FVector& SpaceLinearVel,
@@ -552,7 +564,7 @@ private:
 	UE_API void RemoveClothColliderObjects();
 
 	// Update the sim-space transforms of all cloth collider objects.
-	UE_API void UpdateClothColliderObjects(const FTransform& SpaceTransform);
+	UE_API void UpdateClothColliderObjects();
 
 	// Gather nearby world objects and add them to the sim
 	UE_API void CollectWorldObjects();
@@ -564,7 +576,7 @@ private:
 	UE_API void PurgeExpiredWorldObjects();
 
 	// Update sim-space transforms of world objects
-	UE_API void UpdateWorldObjects(const FTransform& SpaceTransform, const FVector& LinearVel, const FVector& AngularVel, const float DeltaSeconds);
+	UE_API void UpdateWorldObjects(const FVector& LinearVel, const FVector& AngularVel, const float DeltaSeconds);
 
 	// Advances the simulation by a given timestep
 	UE_API void RunPhysicsSimulation(float DeltaSeconds, const FVector& SimSpaceGravity);
@@ -724,6 +736,7 @@ private:
 	FVector PreviousComponentLinearVelocity;
 
 	// Used by the world-space to simulation-space motion transfer system in Component- or Bone-Space sims
+	FTransform SimulationSpaceTransform;
 	FTransform PreviousSimulationSpaceTransform;
 	FTransform PreviousPreviousSimulationSpaceTransform;
 };
diff --git a/Engine/Plugins/Experimental/PhysicsControl/Source/PhysicsControl/Public/PhysicsControlData.h b/Engine/Plugins/Experimental/PhysicsControl/Source/PhysicsControl/Public/PhysicsControlData.h
--- a/Engine/Plugins/Experimental/PhysicsControl/Source/PhysicsControl/Public/PhysicsControlData.h
+++ b/Engine/Plugins/Experimental/PhysicsControl/Source/PhysicsControl/Public/PhysicsControlData.h
@@ -232,6 +232,20 @@ struct FPhysicsControlModifierSparseData
 	uint8 bEnablebUpdateKinematicFromSimulation : 1;
 };
 
+// Specifies a space for queries, forces etc
+UENUM(BlueprintType)
+enum class EPhysicsControlSpace : uint8
+{
+	// World space. For a control rig, this is the space of the control rig, rather than the "real"
+	// world, so is the same as Component.
+	World,
+	// If a skeleton/bone is being used, the target is in the space of that component. If not then
+	// it is in the space of the owning component (e.g. PhysicsControl, or ControlRig etc).
+	Component,
+	// In the space of the body under consideration
+	Body
+};
+
 /**
  * Analogous to the ControlData, this indicates how an individual controlled body should move
  */
@@ -1112,6 +1126,62 @@ struct FPhysicsControlAndBodyModifierCreationDatas
 	TMap<FName, FPhysicsBodyModifierCreationData> Modifiers;
 };
 
+UENUM(BlueprintType)
+enum class EPhysicsControlForceType : uint8
+{
+	Force,
+	Acceleration,
+	Impulse,
+	VelocityChange
+};
+
+UE_API EPhysicsControlForceType ConvertForceType(ImmediatePhysics_Shared::EForceType InForceType);
+UE_API ImmediatePhysics_Shared::EForceType ConvertForceType(EPhysicsControlForceType InForceType);
+
+USTRUCT(BlueprintType)
+struct FPhysicsControlForceAndTorqueData
+{
+	GENERATED_BODY();
+
+	UPROPERTY(EditAnywhere, BlueprintReadWrite, Category = Force)
+	EPhysicsControlSpace LocationSpace = EPhysicsControlSpace::Body;
+
+	UPROPERTY(EditAnywhere, BlueprintReadWrite, Category = Force)
+	EPhysicsControlSpace ForceAndTorqueSpace = EPhysicsControlSpace::World;
+
+	UPROPERTY(EditAnywhere, BlueprintReadWrite, Category = Force)
+	EPhysicsControlForceType Type = EPhysicsControlForceType::Force;
+
+	/// The linear force/acceleration/impulse/velocity change
+	UPROPERTY(EditAnywhere, BlueprintReadWrite, Category = Force)
+	FVector Linear = FVector::ZeroVector;
+
+	/// The angular torque/acceleration/impulse/velocity change
+	UPROPERTY(EditAnywhere, BlueprintReadWrite, Category = Force)
+	FVector Angular = FVector::ZeroVector;
+
+	// The position at which the linear force will be applied
+	UPROPERTY(EditAnywhere, BlueprintReadWrite, Category = Force)
+	FVector Location = FVector::ZeroVector;
+};
+
+USTRUCT(BlueprintType)
+struct FPhysicsControlNamedForceAndTorqueData
+{
+	GENERATED_BODY();
+
+	FPhysicsControlNamedForceAndTorqueData() {}
+
+	FPhysicsControlNamedForceAndTorqueData(const FName InName, const FPhysicsControlForceAndTorqueData& InData)
+		: Name(InName), ForceAndTorqueData(InData) {}
+
+	UPROPERTY(EditAnywhere, BlueprintReadWrite, Category = Force)
+	FName Name;
+
+	UPROPERTY(EditAnywhere, BlueprintReadWrite, Category = Force)
+	FPhysicsControlForceAndTorqueData ForceAndTorqueData;
+};
+
 // Functions to interpolate between various control/modifier datas
 PHYSICSCONTROL_API FPhysicsControlData Interpolate(
 	const FPhysicsControlData& A, const FPhysicsControlData& B, const float Weight);
diff --git a/Engine/Plugins/Experimental/PhysicsControl/Source/PhysicsControl/Public/PhysicsControlObjectVersion.h b/Engine/Plugins/Experimental/PhysicsControl/Source/PhysicsControl/Public/PhysicsControlObjectVersion.h
--- a/Engine/Plugins/Experimental/PhysicsControl/Source/PhysicsControl/Public/PhysicsControlObjectVersion.h
+++ b/Engine/Plugins/Experimental/PhysicsControl/Source/PhysicsControl/Public/PhysicsControlObjectVersion.h
@@ -14,7 +14,7 @@ struct FPhysicsControlObjectVersion
 		ControlRigIncludeDriveInJoint,
 		ControlRigRemoveCurrentDataFromPhysicsComponent,
 		ControlRigUseAutomaticSolver,
-		ControlRigAutomaticallyAddPhysicsComponents,
+		ControlRigAutomaticallyAddPhysicsBodyComponents,
 		ControlRigControlAddChildBodyComponentKey,
 		ControlRigControlAddControlMultiplier,
 		ControlRigSeparateOutJointFromBody,
diff --git a/Engine/Plugins/Experimental/PhysicsControl/Source/PhysicsControl/Public/RigidBodyControlData.h b/Engine/Plugins/Experimental/PhysicsControl/Source/PhysicsControl/Public/RigidBodyControlData.h
--- a/Engine/Plugins/Experimental/PhysicsControl/Source/PhysicsControl/Public/RigidBodyControlData.h
+++ b/Engine/Plugins/Experimental/PhysicsControl/Source/PhysicsControl/Public/RigidBodyControlData.h
@@ -90,6 +90,19 @@ struct FRigidBodyKinematicTargets
 	TMap<FName, FRigidBodyKinematicTarget> Targets;
 };
 
+/**
+ * A set of forces and torques
+ */
+USTRUCT(BlueprintType)
+struct FRigidBodyForcesAndTorques
+{
+	GENERATED_BODY();
+
+	/** Forces (etc) to apply to the named body modifier */
+	UPROPERTY(EditAnywhere, BlueprintReadWrite, Category = PhysicsControl)
+	TMap<FName, FPhysicsControlForceAndTorqueData> ForcesAndTorques;
+};
+
 /**
  * This is the record for a control. It stores the "original" data, which will normally stay fixed.
  * In addition, it will also store some updates to these original data which will determine how the actual
diff --git a/Engine/Source/Runtime/AnimGraphRuntime/Public/BoneControllers/AnimNode_RigidBody.h b/Engine/Source/Runtime/AnimGraphRuntime/Public/BoneControllers/AnimNode_RigidBody.h
--- a/Engine/Source/Runtime/AnimGraphRuntime/Public/BoneControllers/AnimNode_RigidBody.h
+++ b/Engine/Source/Runtime/AnimGraphRuntime/Public/BoneControllers/AnimNode_RigidBody.h
@@ -373,7 +373,8 @@ private:
 		FVector& SpaceLinearAcc,
 		FVector& SpaceAngularAcc);
 
-	// Gather cloth collision sources from the supplied Skeltal Mesh and add a kinematic actor representing each one of them to the sim.
+	// Gather cloth collision sources from the supplied Skeletal Mesh and add a kinematic actor
+	// representing each one of them to the sim.
 	ANIMGRAPHRUNTIME_API void CollectClothColliderObjects(const USkeletalMeshComponent* SkeletalMeshComp);
 	
 	// Remove all cloth collider objects from the sim.
@@ -610,9 +611,9 @@ inline FVector WorldPositionToSpace(
 }
 
 inline FTransform ConvertCSTransformToSimSpace(
-	ESimulationSpace Space, const FTransform& InCSTransform, const FTransform& ComponentToWorld, const FTransform& BaseBoneTM)
+	ESimulationSpace SimSpace, const FTransform& InCSTransform, const FTransform& ComponentToWorld, const FTransform& BaseBoneTM)
 {
-	switch (Space)
+	switch (SimSpace)
 	{
 	case ESimulationSpace::ComponentSpace: return InCSTransform;
 	case ESimulationSpace::WorldSpace:  return InCSTransform * ComponentToWorld;
@@ -620,3 +621,27 @@ inline FTransform ConvertCSTransformToSimSpace(
 	default: ensureMsgf(false, TEXT("Unsupported Simulation Space")); return InCSTransform;
 	}
 }
+
+inline FVector ConvertCSVectorToSimSpace(
+	ESimulationSpace SimSpace, const FVector& InCSVector, const FTransform& ComponentToWorld, const FTransform& BaseBoneTM)
+{
+	switch (SimSpace)
+	{
+	case ESimulationSpace::ComponentSpace: return InCSVector;
+	case ESimulationSpace::WorldSpace:  return ComponentToWorld.TransformVectorNoScale(InCSVector);
+	case ESimulationSpace::BaseBoneSpace: return BaseBoneTM.InverseTransformVectorNoScale(InCSVector); break;
+	default: ensureMsgf(false, TEXT("Unsupported Simulation Space")); return InCSVector;
+	}
+}
+
+inline FVector ConvertCSPositionToSimSpace(
+	ESimulationSpace SimSpace, const FVector& InCSPosition, const FTransform& ComponentToWorld, const FTransform& BaseBoneTM)
+{
+	switch (SimSpace)
+	{
+	case ESimulationSpace::ComponentSpace: return InCSPosition;
+	case ESimulationSpace::WorldSpace:  return ComponentToWorld.TransformPosition(InCSPosition);
+	case ESimulationSpace::BaseBoneSpace: return BaseBoneTM.InverseTransformPosition(InCSPosition); break;
+	default: ensureMsgf(false, TEXT("Unsupported Simulation Space")); return InCSPosition;
+	}
+}
