Add files intended for a previous checkin - just simplification (no change in functionality here)

#Synced-CL 48547444

diff --git a/Engine/Plugins/Experimental/ControlRigPhysics/Source/ControlRigPhysics/Private/RigPhysicsControlExecution.cpp b/Engine/Plugins/Experimental/ControlRigPhysics/Source/ControlRigPhysics/Private/RigPhysicsControlExecution.cpp
--- a/Engine/Plugins/Experimental/ControlRigPhysics/Source/ControlRigPhysics/Private/RigPhysicsControlExecution.cpp
+++ b/Engine/Plugins/Experimental/ControlRigPhysics/Source/ControlRigPhysics/Private/RigPhysicsControlExecution.cpp
@@ -13,6 +13,12 @@
 
 #include UE_INLINE_GENERATED_CPP_BY_NAME(RigPhysicsControlExecution)
 
+//======================================================================================================================
+inline FRigPhysicsControlComponent* GetControl(URigHierarchy* Hierarchy, const FRigComponentKey& Key)
+{
+	return Cast<FRigPhysicsControlComponent>(Hierarchy->FindComponent(Key));
+}
+
 //======================================================================================================================
 FRigUnit_AddPhysicsControl_Execute()
 {
@@ -54,8 +60,7 @@ FRigUnit_HierarchySetControlEnabled_Execute()
 	DECLARE_SCOPE_HIERARCHICAL_COUNTER_RIGUNIT();
 	if (ExecuteContext.Hierarchy)
 	{
-		if (FRigPhysicsControlComponent* Component = Cast<FRigPhysicsControlComponent>(
-			ExecuteContext.Hierarchy->FindComponent(PhysicsControlComponentKey)))
+		if (FRigPhysicsControlComponent* Component = GetControl(ExecuteContext.Hierarchy, PhysicsControlComponentKey))
 		{
 			Component->ControlData.bEnabled = bEnabled;
 		}
@@ -68,8 +73,7 @@ FRigUnit_HierarchySetControlCustomControlPoint_Execute()
 	DECLARE_SCOPE_HIERARCHICAL_COUNTER_RIGUNIT();
 	if (ExecuteContext.Hierarchy)
 	{
-		if (FRigPhysicsControlComponent* Component = Cast<FRigPhysicsControlComponent>(
-			ExecuteContext.Hierarchy->FindComponent(PhysicsControlComponentKey)))
+		if (FRigPhysicsControlComponent* Component = GetControl(ExecuteContext.Hierarchy, PhysicsControlComponentKey))
 		{
 			Component->ControlData.CustomControlPoint = CustomControlPoint;
 			Component->ControlData.bUseCustomControlPoint = bUseCustomControlPoint;
@@ -84,8 +88,7 @@ FRigUnit_HierarchySetControlData_Execute()
 	DECLARE_SCOPE_HIERARCHICAL_COUNTER_RIGUNIT();
 	if (ExecuteContext.Hierarchy)
 	{
-		if (FRigPhysicsControlComponent* Component = Cast<FRigPhysicsControlComponent>(
-			ExecuteContext.Hierarchy->FindComponent(PhysicsControlComponentKey)))
+		if (FRigPhysicsControlComponent* Component = GetControl(ExecuteContext.Hierarchy, PhysicsControlComponentKey))
 		{
 			Component->ControlData = ControlData;
 		}
@@ -98,8 +101,7 @@ FRigUnit_HierarchySetControlLinearStrength_Execute()
 	DECLARE_SCOPE_HIERARCHICAL_COUNTER_RIGUNIT();
 	if (ExecuteContext.Hierarchy)
 	{
-		if (FRigPhysicsControlComponent* Component = Cast<FRigPhysicsControlComponent>(
-			ExecuteContext.Hierarchy->FindComponent(PhysicsControlComponentKey)))
+		if (FRigPhysicsControlComponent* Component = GetControl(ExecuteContext.Hierarchy, PhysicsControlComponentKey))
 		{
 			Component->ControlData.LinearStrength = Strength;
 		}
@@ -112,8 +114,7 @@ FRigUnit_HierarchySetControlLinearDampingRatio_Execute()
 	DECLARE_SCOPE_HIERARCHICAL_COUNTER_RIGUNIT();
 	if (ExecuteContext.Hierarchy)
 	{
-		if (FRigPhysicsControlComponent* Component = Cast<FRigPhysicsControlComponent>(
-			ExecuteContext.Hierarchy->FindComponent(PhysicsControlComponentKey)))
+		if (FRigPhysicsControlComponent* Component = GetControl(ExecuteContext.Hierarchy, PhysicsControlComponentKey))
 		{
 			Component->ControlData.LinearDampingRatio = DampingRatio;
 		}
@@ -126,8 +127,7 @@ FRigUnit_HierarchySetControlTargetVelocityMultipliers_Execute()
 	DECLARE_SCOPE_HIERARCHICAL_COUNTER_RIGUNIT();
 	if (ExecuteContext.Hierarchy)
 	{
-		if (FRigPhysicsControlComponent* Component = Cast<FRigPhysicsControlComponent>(
-			ExecuteContext.Hierarchy->FindComponent(PhysicsControlComponentKey)))
+		if (FRigPhysicsControlComponent* Component = GetControl(ExecuteContext.Hierarchy, PhysicsControlComponentKey))
 		{
 			Component->ControlData.LinearTargetVelocityMultiplier = LinearTargetVelocityMultiplier;
 			Component->ControlData.AngularTargetVelocityMultiplier = AngularTargetVelocityMultiplier;
@@ -141,8 +141,7 @@ FRigUnit_HierarchySetControlAngularStrength_Execute()
 	DECLARE_SCOPE_HIERARCHICAL_COUNTER_RIGUNIT();
 	if (ExecuteContext.Hierarchy)
 	{
-		if (FRigPhysicsControlComponent* Component = Cast<FRigPhysicsControlComponent>(
-			ExecuteContext.Hierarchy->FindComponent(PhysicsControlComponentKey)))
+		if (FRigPhysicsControlComponent* Component = GetControl(ExecuteContext.Hierarchy, PhysicsControlComponentKey))
 		{
 			Component->ControlData.AngularStrength = Strength;
 		}
@@ -155,8 +154,7 @@ FRigUnit_HierarchySetControlAngularDampingRatio_Execute()
 	DECLARE_SCOPE_HIERARCHICAL_COUNTER_RIGUNIT();
 	if (ExecuteContext.Hierarchy)
 	{
-		if (FRigPhysicsControlComponent* Component = Cast<FRigPhysicsControlComponent>(
-			ExecuteContext.Hierarchy->FindComponent(PhysicsControlComponentKey)))
+		if (FRigPhysicsControlComponent* Component = GetControl(ExecuteContext.Hierarchy, PhysicsControlComponentKey))
 		{
 			Component->ControlData.AngularDampingRatio = DampingRatio;
 		}
@@ -169,8 +167,7 @@ FRigUnit_HierarchyGetControlData_Execute()
 	DECLARE_SCOPE_HIERARCHICAL_COUNTER_RIGUNIT();
 	if (ExecuteContext.Hierarchy)
 	{
-		if (FRigPhysicsControlComponent* Component = Cast<FRigPhysicsControlComponent>(
-			ExecuteContext.Hierarchy->FindComponent(PhysicsControlComponentKey)))
+		if (FRigPhysicsControlComponent* Component = GetControl(ExecuteContext.Hierarchy, PhysicsControlComponentKey))
 		{
 			ControlData = Component->ControlData;
 		}
@@ -183,8 +180,7 @@ FRigUnit_HierarchySetControlMultiplier_Execute()
 	DECLARE_SCOPE_HIERARCHICAL_COUNTER_RIGUNIT();
 	if (ExecuteContext.Hierarchy)
 	{
-		if (FRigPhysicsControlComponent* Component = Cast<FRigPhysicsControlComponent>(
-			ExecuteContext.Hierarchy->FindComponent(PhysicsControlComponentKey)))
+		if (FRigPhysicsControlComponent* Component = GetControl(ExecuteContext.Hierarchy, PhysicsControlComponentKey))
 		{
 			Component->ControlMultiplier = ControlMultiplier;
 		}
@@ -197,8 +193,7 @@ FRigUnit_HierarchySetControlDataAndMultiplier_Execute()
 	DECLARE_SCOPE_HIERARCHICAL_COUNTER_RIGUNIT();
 	if (ExecuteContext.Hierarchy)
 	{
-		if (FRigPhysicsControlComponent* Component = Cast<FRigPhysicsControlComponent>(
-			ExecuteContext.Hierarchy->FindComponent(PhysicsControlComponentKey)))
+		if (FRigPhysicsControlComponent* Component = GetControl(ExecuteContext.Hierarchy, PhysicsControlComponentKey))
 		{
 			Component->ControlData = ControlData;
 			Component->ControlMultiplier = ControlMultiplier;
@@ -212,8 +207,7 @@ FRigUnit_HierarchySetControlTarget_Execute()
 	DECLARE_SCOPE_HIERARCHICAL_COUNTER_RIGUNIT();
 	if (ExecuteContext.Hierarchy)
 	{
-		if (FRigPhysicsControlComponent* Component = Cast<FRigPhysicsControlComponent>(
-			ExecuteContext.Hierarchy->FindComponent(PhysicsControlComponentKey)))
+		if (FRigPhysicsControlComponent* Component = GetControl(ExecuteContext.Hierarchy, PhysicsControlComponentKey))
 		{
 			Component->ControlTarget = ControlTarget;
 		}
@@ -226,8 +220,7 @@ FRigUnit_HierarchyUpdateControlTarget_Execute()
 	DECLARE_SCOPE_HIERARCHICAL_COUNTER_RIGUNIT();
 	if (ExecuteContext.Hierarchy)
 	{
-		if (FRigPhysicsControlComponent* Component = Cast<FRigPhysicsControlComponent>(
-			ExecuteContext.Hierarchy->FindComponent(PhysicsControlComponentKey)))
+		if (FRigPhysicsControlComponent* Component = GetControl(ExecuteContext.Hierarchy, PhysicsControlComponentKey))
 		{
 			Component->ControlTarget.TargetVelocity = UE::PhysicsControl::CalculateLinearVelocity(
 				TargetPosition, Component->ControlTarget.TargetPosition, DeltaTime);
diff --git a/Engine/Plugins/Experimental/ControlRigPhysics/Source/ControlRigPhysics/Private/RigPhysicsJointExecution.cpp b/Engine/Plugins/Experimental/ControlRigPhysics/Source/ControlRigPhysics/Private/RigPhysicsJointExecution.cpp
--- a/Engine/Plugins/Experimental/ControlRigPhysics/Source/ControlRigPhysics/Private/RigPhysicsJointExecution.cpp
+++ b/Engine/Plugins/Experimental/ControlRigPhysics/Source/ControlRigPhysics/Private/RigPhysicsJointExecution.cpp
@@ -20,6 +20,12 @@
 
 #include UE_INLINE_GENERATED_CPP_BY_NAME(RigPhysicsJointExecution)
 
+//======================================================================================================================
+inline FRigPhysicsJointComponent* GetJoint(URigHierarchy* Hierarchy, const FRigComponentKey& Key)
+{
+	return Cast<FRigPhysicsJointComponent>(Hierarchy->FindComponent(Key));
+}
+
 //======================================================================================================================
 FRigUnit_AddPhysicsJoint_Execute()
 {
@@ -60,8 +66,7 @@ FRigUnit_HierarchySetJointData_Execute()
 	{
 		return;
 	}
-	if (FRigPhysicsJointComponent* Component = Cast<FRigPhysicsJointComponent>(
-		ExecuteContext.Hierarchy->FindComponent(PhysicsJointComponentKey)))
+	if (FRigPhysicsJointComponent* Component = GetJoint(ExecuteContext.Hierarchy, PhysicsJointComponentKey))
 	{
 		Component->JointData = JointData;
 	}
@@ -75,8 +80,7 @@ FRigUnit_HierarchyGetJointData_Execute()
 	{
 		return;
 	}
-	if (FRigPhysicsJointComponent* Component = Cast<FRigPhysicsJointComponent>(
-		ExecuteContext.Hierarchy->FindComponent(PhysicsJointComponentKey)))
+	if (FRigPhysicsJointComponent* Component = GetJoint(ExecuteContext.Hierarchy, PhysicsJointComponentKey))
 	{
 		JointData = Component->JointData;
 	}
@@ -90,8 +94,7 @@ FRigUnit_HierarchySetJointLinearLimit_Execute()
 	{
 		return;
 	}
-	if (FRigPhysicsJointComponent* Component = Cast<FRigPhysicsJointComponent>(
-		ExecuteContext.Hierarchy->FindComponent(PhysicsJointComponentKey)))
+	if (FRigPhysicsJointComponent* Component = GetJoint(ExecuteContext.Hierarchy, PhysicsJointComponentKey))
 	{
 		Component->JointData.LinearConstraint = LinearLimit.LinearConstraint;
 	}
@@ -105,8 +108,7 @@ FRigUnit_HierarchySetJointConeLimit_Execute()
 	{
 		return;
 	}
-	if (FRigPhysicsJointComponent* Component = Cast<FRigPhysicsJointComponent>(
-		ExecuteContext.Hierarchy->FindComponent(PhysicsJointComponentKey)))
+	if (FRigPhysicsJointComponent* Component = GetJoint(ExecuteContext.Hierarchy, PhysicsJointComponentKey))
 	{
 		Component->JointData.ConeConstraint = ConeLimit.ConeConstraint;
 	}
@@ -120,8 +122,7 @@ FRigUnit_HierarchySetJointTwistLimit_Execute()
 	{
 		return;
 	}
-	if (FRigPhysicsJointComponent* Component = Cast<FRigPhysicsJointComponent>(
-		ExecuteContext.Hierarchy->FindComponent(PhysicsJointComponentKey)))
+	if (FRigPhysicsJointComponent* Component = GetJoint(ExecuteContext.Hierarchy, PhysicsJointComponentKey))
 	{
 		Component->JointData.TwistConstraint = TwistLimit.TwistConstraint;
 	}
@@ -135,8 +136,7 @@ FRigUnit_HierarchySetJointEnabled_Execute()
 	{
 		return;
 	}
-	if (FRigPhysicsJointComponent* Component = Cast<FRigPhysicsJointComponent>(
-		ExecuteContext.Hierarchy->FindComponent(PhysicsJointComponentKey)))
+	if (FRigPhysicsJointComponent* Component = GetJoint(ExecuteContext.Hierarchy, PhysicsJointComponentKey))
 	{
 		Component->JointData.bEnabled = bEnabled;
 	}
@@ -150,8 +150,7 @@ FRigUnit_HierarchySetJointDriveData_Execute()
 	{
 		return;
 	}
-	if (FRigPhysicsJointComponent* Component = Cast<FRigPhysicsJointComponent>(
-		ExecuteContext.Hierarchy->FindComponent(PhysicsJointComponentKey)))
+	if (FRigPhysicsJointComponent* Component = GetJoint(ExecuteContext.Hierarchy, PhysicsJointComponentKey))
 	{
 		Component->DriveData = DriveData;
 	}
@@ -165,8 +164,7 @@ FRigUnit_HierarchyGetJointDriveData_Execute()
 	{
 		return;
 	}
-	if (FRigPhysicsJointComponent* Component = Cast<FRigPhysicsJointComponent>(
-		ExecuteContext.Hierarchy->FindComponent(PhysicsJointComponentKey)))
+	if (FRigPhysicsJointComponent* Component = GetJoint(ExecuteContext.Hierarchy, PhysicsJointComponentKey))
 	{
 		DriveData = Component->DriveData;
 	}
@@ -180,8 +178,7 @@ FRigUnit_HierarchySetJointDriveUseSkeletalAnimation_Execute()
 	{
 		return;
 	}
-	if (FRigPhysicsJointComponent* Component = Cast<FRigPhysicsJointComponent>(
-		ExecuteContext.Hierarchy->FindComponent(PhysicsJointComponentKey)))
+	if (FRigPhysicsJointComponent* Component = GetJoint(ExecuteContext.Hierarchy, PhysicsJointComponentKey))
 	{
 		Component->DriveData.bUseSkeletalAnimation = bUseSkeletalAnimation;
 	}
@@ -381,4 +378,4 @@ FRigUnit_MakeDriveData_Execute()
 		DriveData.AngularDriveConstraint.SetAngularVelocityDriveTwistAndSwing(false, false);
 		DriveData.AngularDriveConstraint.SetAngularVelocityDriveSLERP(false);
 	}
-}
\ No newline at end of file
+}
