Fix missing initialisation in FRigUnit_HierarchyGetPhysicsBodyPointVelocity

#Synced-CL 48717244

diff --git a/Engine/Plugins/Experimental/ControlRigPhysics/Source/ControlRigPhysics/Public/RigPhysicsBodyExecution.h b/Engine/Plugins/Experimental/ControlRigPhysics/Source/ControlRigPhysics/Public/RigPhysicsBodyExecution.h
--- a/Engine/Plugins/Experimental/ControlRigPhysics/Source/ControlRigPhysics/Public/RigPhysicsBodyExecution.h
+++ b/Engine/Plugins/Experimental/ControlRigPhysics/Source/ControlRigPhysics/Public/RigPhysicsBodyExecution.h
@@ -699,7 +699,7 @@ struct FRigUnit_HierarchyGetPhysicsBodyPointVelocity : public FRigUnit_PhysicsBa
 
 	// The position to calculate the velocity at.
 	UPROPERTY(meta = (Input))
-	FVector Position;
+	FVector Position = FVector::ZeroVector;
 
 	// The space that Position is specified in.
 	UPROPERTY(meta = (Input))
