Add a redirector for RigPhysicsSolverSettings.bAutomaticallyAddPhysicsComponents

#Synced-CL 48844440

diff --git a/Engine/Plugins/Experimental/ControlRigPhysics/Config/DefaultControlRigPhysics.ini b/Engine/Plugins/Experimental/ControlRigPhysics/Config/DefaultControlRigPhysics.ini
new file mode 100644
--- /dev/null
+++ b/Engine/Plugins/Experimental/ControlRigPhysics/Config/DefaultControlRigPhysics.ini
@@ -0,0 +1,2 @@
+[CoreRedirects]
++PropertyRedirects=(OldName="RigPhysicsSolverSettings.bAutomaticallyAddPhysicsComponents",NewName="bAutomaticallyAddPhysicsBodyComponents")
