Control Rig Physics - Improve instantiation from a physics assert, to let modular rigs work better.

#Synced-CL 49169110

diff --git a/Engine/Plugins/Experimental/ControlRigPhysics/Source/ControlRigPhysics/Private/RigPhysicsExecution.cpp b/Engine/Plugins/Experimental/ControlRigPhysics/Source/ControlRigPhysics/Private/RigPhysicsExecution.cpp
--- a/Engine/Plugins/Experimental/ControlRigPhysics/Source/ControlRigPhysics/Private/RigPhysicsExecution.cpp
+++ b/Engine/Plugins/Experimental/ControlRigPhysics/Source/ControlRigPhysics/Private/RigPhysicsExecution.cpp
@@ -335,6 +335,35 @@ static void SetUpJointData(
 	JointData.AngularProjectionAmount = ConstraintProfileProperties.ProjectionAngularAlpha;
 }
 
+//======================================================================================================================
+static FRigComponentKey FindPhysicsBodyComponentWithSolver(
+	const URigHierarchy* Hierarchy, const FRigElementKey& ElementKey, const FRigComponentKey& SolverKey)
+{
+	// Note that our owning element might have multiple bodies, and we can't be sure of the name. So
+	// look for one that has the same solver as us - there should only be one of those.
+	const TArray<FRigComponentKey> CandidateBodyComponentKeys =
+		Hierarchy->GetComponentKeys(ElementKey);
+	if (const FRigComponentKey* CC = CandidateBodyComponentKeys.FindByPredicate(
+		[SolverKey, &Hierarchy](const FRigComponentKey& Other) {
+			if (const FRigPhysicsBodyComponent* PBC = Cast<FRigPhysicsBodyComponent>(
+				Hierarchy->FindComponent(Other)))
+			{
+				if (PBC->BodySolverSettings.bUseAutomaticSolver ||
+					PBC->BodySolverSettings.PhysicsSolverComponentKey == SolverKey)
+				{
+					return true;
+				}
+			}
+			return false;
+		}))
+	{
+		return *CC;
+	}
+
+	return FRigComponentKey();
+
+}
+
 //======================================================================================================================
 FRigUnit_HierarchyInstantiateFromPhysicsAsset_Execute()
 {
@@ -372,88 +401,172 @@ FRigUnit_HierarchyInstantiateFromPhysicsAsset_Execute()
 			continue;
 		}
 
+		// If this flag is set then it means the body under consideration is a helper body - a
+		// kinematic one that is not in BonesToUse, but is needed to support the bodies we should be creating.
+		bool bNeedBodyAsHelper = false;
+
 		// Check we can use this based on the bone name
-		FName BoneName = SkeletalBodySetup->BoneName;
-		FRigElementKey OwnerElementKey(BoneName, ERigElementType::Bone);
+		const FName BoneName = SkeletalBodySetup->BoneName;
+		const FRigElementKey OwnerElementKey(BoneName, ERigElementType::Bone);
 		if (!ExecuteContext.Hierarchy->Find(OwnerElementKey))
 		{
 			continue;
 		}
 		if (!BonesToUse.IsEmpty() && BonesToUse.Find(OwnerElementKey) == INDEX_NONE)
 		{
-			continue;
-		}
+			if (bMakeHelperBodies)
+			{
+				// Whilst this bone isn't in the list of bones to use, it may be connected via a
+				// constraint (physics joint) to a bone that is in use. If we find that this is the
+				// case, we can make a kinematic body for it (flagged as a helper body) so that
+				// we're able to form joints and controls between it and bodies that are in BonesToUse.
 
-		// Collision
-		FRigPhysicsCollision Collision;
-		const FKAggregateGeom& AggGeom = SkeletalBodySetup->AggGeom;
-		for (const FKBoxElem& BoxElem : AggGeom.BoxElems)
-		{
-			FTransform TM(BoxElem.Rotation, BoxElem.Center);
-			FRigPhysicsCollisionBox Box(TM, FVector(BoxElem.X, BoxElem.Y, BoxElem.Z));
-			SetCommonProperties(Box, BoxElem);
-			Collision.Boxes.Add(Box);
-		}
-		for (const FKSphereElem& SphereElem : AggGeom.SphereElems)
-		{
-			FTransform TM(SphereElem.Center);
-			FRigPhysicsCollisionSphere Sphere(TM, SphereElem.Radius);
-			SetCommonProperties(Sphere, SphereElem);
-			Collision.Spheres.Add(Sphere);
-		}
-		for (const FKSphylElem& SphylElem : AggGeom.SphylElems)
-		{
-			FTransform TM(SphylElem.Rotation, SphylElem.Center);
-			FRigPhysicsCollisionCapsule Capsule(TM, SphylElem.Radius, SphylElem.Length);
-			SetCommonProperties(Capsule, SphylElem);
-			Collision.Capsules.Add(Capsule);
-		}
+				// Danny TODO this nested O(N) search is not ideal
+				for (const TObjectPtr<UPhysicsConstraintTemplate>& ConstraintTemplate : PhysicsAsset->ConstraintSetup)
+				{
+					if (!ConstraintTemplate)
+					{
+						continue;
+					}
+					FConstraintInstance& ConstraintInstance = ConstraintTemplate->DefaultInstance;
+					const FName ChildBoneName = ConstraintInstance.ConstraintBone1;
+					const FName ParentBoneName = ConstraintInstance.ConstraintBone2;
 
-		const UPhysicalMaterial* Material = SkeletalBodySetup->GetPhysMaterial();
-		Material = Material ? Material : GetDefault<UPhysicalMaterial>();
-		if (Material)
-		{
-			Collision.Material.Friction = FMath::Max(Material->Friction, Material->StaticFriction);
-			Collision.Material.Restitution = Material->Restitution;
-			Collision.Material.FrictionCombineMode = (ERigPhysicsCombineMode) 
-				Material->FrictionCombineMode.GetValue();
-			Collision.Material.RestitutionCombineMode = (ERigPhysicsCombineMode) 
-				Material->RestitutionCombineMode.GetValue();
+					if (BoneName == ParentBoneName &&
+						ExecuteContext.Hierarchy->Find(FRigElementKey(ChildBoneName, ERigElementType::Bone)) &&
+						BonesToUse.Find(FRigElementKey(ChildBoneName, ERigElementType::Bone)) != INDEX_NONE)
+					{
+						bNeedBodyAsHelper = true;
+						break;
+					}
+				}
+
+				if (!bNeedBodyAsHelper)
+				{
+					continue;
+				}
+			}
+			else
+			{
+				continue;
+			}
 		}
 
-		// Dynamics
-		FRigPhysicsDynamics Dynamics;
-		Dynamics.MassOverride = SkeletalBodySetup->CalculateMass();
-		Dynamics.LinearDamping = SkeletalBodySetup->DefaultInstance.LinearDamping;
-		Dynamics.AngularDamping = SkeletalBodySetup->DefaultInstance.AngularDamping;
-		Dynamics.CentreOfMassNudge = SkeletalBodySetup->DefaultInstance.COMNudge;
-
-		// Body Data
-		FPhysicsControlModifierData BodyData;
-		BodyData.MovementType =
-			SkeletalBodySetup->PhysicsType == EPhysicsType::PhysType_Simulated ?
-			EPhysicsMovementType::Simulated : EPhysicsMovementType::Kinematic;
-		BodyData.CollisionType =
-			SkeletalBodySetup->CollisionReponse == EBodyCollisionResponse::BodyCollision_Enabled ?
-			ECollisionEnabled::PhysicsOnly : ECollisionEnabled::NoCollision;
-
-		// Add the component
-		FRigComponentKey PhysicsBodyComponentKey = Controller->AddComponent(
-			FRigPhysicsBodyComponent::StaticStruct(), TEXT("PhysicsBody"), OwnerElementKey);
-		if (FRigPhysicsBodyComponent* Component = Cast<FRigPhysicsBodyComponent>(
-			ExecuteContext.Hierarchy->FindComponent(PhysicsBodyComponentKey)))
+		// This is the body component that we will (one of):
+		// * Make from scratch
+		// * Re-use, but update
+		// * Re-use as-is
+
+		// Note that this will return an invalid key if an existing body component can't be found
+		FRigComponentKey PhysicsBodyComponentKey = FindPhysicsBodyComponentWithSolver(
+			ExecuteContext.Hierarchy, OwnerElementKey, PhysicsSolverComponentKey);
+
+		// Note that the body (PhysicsBodyComponentKey) might exist, but even then it might have
+		// been created as a helper. In that case, we don't need to recreate it, but we do need to
+		// overwrite its properties.
+
+		// If the key is not valid, we will make the body. But we also need to decide if we will
+		// overwrite the collision/dynamics properties, and whether we force it to be kinematic.
+
+		// Need helper      Key Valid       Make body   Overwrite properties   Force kinematic
+		// 0                0               1           1                      0
+		// 0                1               0           1                      0
+		// 1                0               1           1                      1
+		// 1                1               0           0                      0
+
+		if (!(PhysicsBodyComponentKey.IsValid() && bNeedBodyAsHelper))
 		{
-			Component->BodySolverSettings.PhysicsSolverComponentKey = PhysicsSolverComponentKey;
-			Component->BodySolverSettings.bUseAutomaticSolver = bUseAutomaticSolver;
-			Component->Dynamics = Dynamics;
-			Component->BodyData = BodyData;
-			Component->Collision = Collision;
+			FRigPhysicsCollision Collision;
+			FRigPhysicsDynamics Dynamics;
+			FPhysicsControlModifierData BodyData;
 
+			// Collision
+			const FKAggregateGeom& AggGeom = SkeletalBodySetup->AggGeom;
+			for (const FKBoxElem& BoxElem : AggGeom.BoxElems)
+			{
+				FTransform TM(BoxElem.Rotation, BoxElem.Center);
+				FRigPhysicsCollisionBox Box(TM, FVector(BoxElem.X, BoxElem.Y, BoxElem.Z));
+				SetCommonProperties(Box, BoxElem);
+				Collision.Boxes.Add(Box);
+			}
+			for (const FKSphereElem& SphereElem : AggGeom.SphereElems)
+			{
+				FTransform TM(SphereElem.Center);
+				FRigPhysicsCollisionSphere Sphere(TM, SphereElem.Radius);
+				SetCommonProperties(Sphere, SphereElem);
+				Collision.Spheres.Add(Sphere);
+			}
+			for (const FKSphylElem& SphylElem : AggGeom.SphylElems)
+			{
+				FTransform TM(SphylElem.Rotation, SphylElem.Center);
+				FRigPhysicsCollisionCapsule Capsule(TM, SphylElem.Radius, SphylElem.Length);
+				SetCommonProperties(Capsule, SphylElem);
+				Collision.Capsules.Add(Capsule);
+			}
+
+			const UPhysicalMaterial* Material = SkeletalBodySetup->GetPhysMaterial();
+			Material = Material ? Material : GetDefault<UPhysicalMaterial>();
+			if (Material)
+			{
+				Collision.Material.Friction = FMath::Max(Material->Friction, Material->StaticFriction);
+				Collision.Material.Restitution = Material->Restitution;
+				Collision.Material.FrictionCombineMode = (ERigPhysicsCombineMode)
+					Material->FrictionCombineMode.GetValue();
+				Collision.Material.RestitutionCombineMode = (ERigPhysicsCombineMode)
+					Material->RestitutionCombineMode.GetValue();
+			}
+
+			// Dynamics
+			Dynamics.MassOverride = SkeletalBodySetup->CalculateMass();
+			Dynamics.LinearDamping = SkeletalBodySetup->DefaultInstance.LinearDamping;
+			Dynamics.AngularDamping = SkeletalBodySetup->DefaultInstance.AngularDamping;
+			Dynamics.CentreOfMassNudge = SkeletalBodySetup->DefaultInstance.COMNudge;
+
+			// Body Data
+			BodyData.MovementType =
+				SkeletalBodySetup->PhysicsType == EPhysicsType::PhysType_Simulated ?
+				EPhysicsMovementType::Simulated : EPhysicsMovementType::Kinematic;
+			BodyData.CollisionType =
+				SkeletalBodySetup->CollisionReponse == EBodyCollisionResponse::BodyCollision_Enabled ?
+				ECollisionEnabled::PhysicsOnly : ECollisionEnabled::NoCollision;
+
+			if (bNeedBodyAsHelper)
+			{
+				BodyData.MovementType = EPhysicsMovementType::Kinematic;
+				BodyData.CollisionType = ECollisionEnabled::NoCollision;
+			}
+
+			// Add the component if necessary
+			if (!PhysicsBodyComponentKey.IsValid())
+			{
+				PhysicsBodyComponentKey = Controller->AddComponent(
+					FRigPhysicsBodyComponent::StaticStruct(), TEXT("PhysicsBody"), OwnerElementKey);
+			}
+
+			// Update the component's properties
+			if (FRigPhysicsBodyComponent* Component = Cast<FRigPhysicsBodyComponent>(
+				ExecuteContext.Hierarchy->FindComponent(PhysicsBodyComponentKey)))
+			{
+				Component->BodySolverSettings.PhysicsSolverComponentKey = PhysicsSolverComponentKey;
+				Component->BodySolverSettings.bUseAutomaticSolver = bUseAutomaticSolver;
+				Component->Dynamics = Dynamics;
+				Component->BodyData = BodyData;
+				Component->Collision = Collision;
+
+				BoneToBodyComponentKeyMap.Add(BoneName, PhysicsBodyComponentKey);
+				if (!bNeedBodyAsHelper)
+				{
+					PhysicsBodyComponentKeys.Add(PhysicsBodyComponentKey);
+				}
+			}
+		} // Creating or updating the body
+		else
+		{
 			BoneToBodyComponentKeyMap.Add(BoneName, PhysicsBodyComponentKey);
-			PhysicsBodyComponentKeys.Add(PhysicsBodyComponentKey);
 		}
 
-		if (bAddSimSpaceControl)
+		// Sim Space control
+		if (bAddSimSpaceControl && !bNeedBodyAsHelper)
 		{
 			FRigComponentKey SimSpaceControlComponentKey = Controller->AddComponent(
 				FRigPhysicsControlComponent::StaticStruct(),
@@ -469,14 +582,13 @@ FRigUnit_HierarchyInstantiateFromPhysicsAsset_Execute()
 			}
 		}
 
-		if (bAddParentSpaceControl)
+		// Parent space control
+		if (bAddParentSpaceControl && !bNeedBodyAsHelper)
 		{
-			// Danny TODO Note that we create a parent-space control for every body - even the root
-			// body. This is often not desired, so we should provide a way of disabling it, as
-			// otherwise it will provide slightly unexpected behaviour (even though it's "correct")
-			// as we end up with a sim-space control if there is no parent body. We can't simply
-			// move this to the joint creation loop, since there may be some joints where we don't
-			// want a parent-space control (e.g. hamstrings).
+			// Note that we create a parent-space control for every body - even the root-most body.
+			// This shouldn't be a problem because we'll set bUseParentBodyAsDefault, and when the
+			// control is instantiated, this will be used to distinguish between a genuine sim-space
+			// control, and a parent-space control where there is no parent body.
 			FRigComponentKey ParentSpaceControlComponentKey = Controller->AddComponent(
 				FRigPhysicsControlComponent::StaticStruct(),
 				TEXT("ParentSpaceControl"), PhysicsBodyComponentKey.ElementKey);
@@ -493,8 +605,12 @@ FRigUnit_HierarchyInstantiateFromPhysicsAsset_Execute()
 	} // Loop over SkeletalBodySetup
 
 	// Now bodies are created we can handle no-collision
+	//
+	// Danny TODO Note that this is a loop over the whole collision table in the physics asset -
+	// even if we are just pulling out a small number of bones from it. It could potentially be
+	// quite expensive, especially as once we have a pair of bones, we need to iterate further to
+	// decide whether to keep the pair (especially as we don't want to make duplicates).
 	const TMap<FRigidBodyIndexPair, bool>& CollisionDisableTable = PhysicsAsset->CollisionDisableTable;
-
 	for (const TPair<FRigidBodyIndexPair, bool>& PairPair : CollisionDisableTable)
 	{
 		if (!PairPair.Value)
@@ -509,17 +625,46 @@ FRigUnit_HierarchyInstantiateFromPhysicsAsset_Execute()
 				{
 					const FName BoneName1 = PhysicsAsset->SkeletalBodySetups[I1]->BoneName;
 					const FName BoneName2 = PhysicsAsset->SkeletalBodySetups[I2]->BoneName;
+					bool bInclude1 = (BoneToBodyComponentKeyMap.Find(BoneName1) != nullptr);
+					bool bInclude2 = (BoneToBodyComponentKeyMap.Find(BoneName2) != nullptr);
 
-					if (FRigComponentKey* BodyComponentKey1 = BoneToBodyComponentKeyMap.Find(BoneName1))
+					if (bInclude1 || bInclude2)
 					{
-						if (FRigPhysicsBodyComponent* Component = Cast<FRigPhysicsBodyComponent>(
-							ExecuteContext.Hierarchy->FindComponent(*BodyComponentKey1)))
-
-						if (FRigComponentKey* BodyComponentKey2 = BoneToBodyComponentKeyMap.Find(BoneName2))
+						// If the user instantiates from the same physics asset multiple times, for
+						// different limbs, and bodies from one limb are set to not collide with the
+						// other limbs, then during the first limb processing, the other limb bodies
+						// won't exist yet. We don't try to register the no-collision no - we wait
+						// until the second (etc) limbs are created.
+						//
+						// Danny TODO We could skip a lot of these checks/loops if we're creating
+						// from the whole physics asset - i.e. it might be worth having an
+						// alternative path for that.
+						FRigComponentKey BodyComponentKey1 = FindPhysicsBodyComponentWithSolver(
+							ExecuteContext.Hierarchy, 
+							FRigElementKey(BoneName1, ERigElementType::Bone), 
+							PhysicsSolverComponentKey);
+						if (BodyComponentKey1.IsValid())
 						{
-							if (ExecuteContext.Hierarchy->FindComponent(*BodyComponentKey2))
+							FRigComponentKey BodyComponentKey2 = FindPhysicsBodyComponentWithSolver(
+								ExecuteContext.Hierarchy,
+								FRigElementKey(BoneName2, ERigElementType::Bone),
+								PhysicsSolverComponentKey);
+							if (BodyComponentKey2.IsValid())
 							{
-								Component->NoCollisionBodies.Add(*BodyComponentKey2);
+								FRigPhysicsBodyComponent* Component1 = Cast<FRigPhysicsBodyComponent>(
+									ExecuteContext.Hierarchy->FindComponent(BodyComponentKey1));
+								FRigPhysicsBodyComponent* Component2 = Cast<FRigPhysicsBodyComponent>(
+									ExecuteContext.Hierarchy->FindComponent(BodyComponentKey2));
+								// It probably doesn't really matter, but add the no-collision to
+								// the component we're actually working with
+								if (bInclude1 && Component1)
+								{
+									Component1->NoCollisionBodies.AddUnique(BodyComponentKey2);
+								}
+								else if (Component2)
+								{
+									Component2->NoCollisionBodies.AddUnique(BodyComponentKey1);
+								}
 							}
 						}
 					}
@@ -531,7 +676,7 @@ FRigUnit_HierarchyInstantiateFromPhysicsAsset_Execute()
 	// Physics Joints
 	if (bEnableJoints)
 	{
-		for (TObjectPtr<UPhysicsConstraintTemplate>& ConstraintTemplate : PhysicsAsset->ConstraintSetup)
+		for (const TObjectPtr<UPhysicsConstraintTemplate>& ConstraintTemplate : PhysicsAsset->ConstraintSetup)
 		{
 			if (!ConstraintTemplate)
 			{
diff --git a/Engine/Plugins/Experimental/ControlRigPhysics/Source/ControlRigPhysics/Private/RigPhysicsSimulation.cpp b/Engine/Plugins/Experimental/ControlRigPhysics/Source/ControlRigPhysics/Private/RigPhysicsSimulation.cpp
--- a/Engine/Plugins/Experimental/ControlRigPhysics/Source/ControlRigPhysics/Private/RigPhysicsSimulation.cpp
+++ b/Engine/Plugins/Experimental/ControlRigPhysics/Source/ControlRigPhysics/Private/RigPhysicsSimulation.cpp
@@ -900,9 +900,23 @@ void FRigPhysicsSimulation::InstantiateControls(
 		{
 			if (FRigBodyRecord* ChildBodyRecord = BodyRecords.Find(ControlRecord.ChildBodyComponentKey))
 			{
-				// this can be nullptr - just means a global control
 				FRigBodyRecord* ParentBodyRecord = BodyRecords.Find(ControlRecord.ParentBodyComponentKey);
 
+				// ParentBodyRecord being nullptr just means to create a sim-space control. However,
+				// we never want to make a sim-space control if the user has set
+				// bUseParentBodyAsDefault, because that flag implies the user wanted a parent-space
+				// control. Note that when creating chains etc it's quite easy to "accidentally"
+				// create a parent-space control for the first body in the chain, even when one
+				// isn't actually required. In fact, this can happen very easily with modular rigs
+				// because chains may be created independently, and we'll only find out we want a
+				// parent-space control on the arm (say) after the spine has been created.
+				if (!ParentBodyRecord && ControlComponent->bUseParentBodyAsDefault)
+				{
+					// We will have a control in the component, but it will do nothing as there will
+					// be no physics joint/constraint.
+					continue;
+				}
+
 				ImmediatePhysics::FActorHandle* const ChildBodyHandle = ChildBodyRecord->ActorHandle;
 				ImmediatePhysics::FActorHandle* const ParentBodyHandle = 
 					ParentBodyRecord ? ParentBodyRecord->ActorHandle : SimulationActorHandle;
diff --git a/Engine/Plugins/Experimental/ControlRigPhysics/Source/ControlRigPhysics/Public/RigPhysicsExecution.h b/Engine/Plugins/Experimental/ControlRigPhysics/Source/ControlRigPhysics/Public/RigPhysicsExecution.h
--- a/Engine/Plugins/Experimental/ControlRigPhysics/Source/ControlRigPhysics/Public/RigPhysicsExecution.h
+++ b/Engine/Plugins/Experimental/ControlRigPhysics/Source/ControlRigPhysics/Public/RigPhysicsExecution.h
@@ -453,6 +453,14 @@ struct FRigUnit_HierarchyInstantiateFromPhysicsAsset : public FRigUnit_PhysicsBa
 	UPROPERTY(meta = (Input))
 	bool bAddParentSpaceControl = false;
 
+	// If true, and BonesToUse is being used to create a subset of the physics asset, then helper
+	// bodies will be created when needed: Kinematic bodies on bones that are connected by a physics
+	// joint to those in BonesToUse. This is used so that (a) instantiating individual chains can be
+	// done without explicitly including the parent (in the physics joint sense) body, and (b)
+	// entire physics assets can be made from multiple calls to instantiate individual chains.
+	UPROPERTY(meta = (Input))
+	bool bMakeHelperBodies = true;
+
 	// Data for the simulation space control
 	UPROPERTY(meta = (Input))
 	FPhysicsControlData SimSpaceControlData;
