Control rig physics - some change to not assume constraints are set up in a partciular order

#Synced-CL 49169110

diff --git a/Engine/Plugins/Experimental/ControlRigPhysics/Source/ControlRigPhysics/Private/RigPhysicsExecution.cpp b/Engine/Plugins/Experimental/ControlRigPhysics/Source/ControlRigPhysics/Private/RigPhysicsExecution.cpp
--- a/Engine/Plugins/Experimental/ControlRigPhysics/Source/ControlRigPhysics/Private/RigPhysicsExecution.cpp
+++ b/Engine/Plugins/Experimental/ControlRigPhysics/Source/ControlRigPhysics/Private/RigPhysicsExecution.cpp
@@ -439,6 +439,15 @@ FRigUnit_HierarchyInstantiateFromPhysicsAsset_Execute()
 						bNeedBodyAsHelper = true;
 						break;
 					}
+
+					// Constraints can also be set up in the other order (though less commonly)
+					if (BoneName == ChildBoneName &&
+						ExecuteContext.Hierarchy->Find(FRigElementKey(ParentBoneName, ERigElementType::Bone)) &&
+						BonesToUse.Find(FRigElementKey(ParentBoneName, ERigElementType::Bone)) != INDEX_NONE)
+					{
+						bNeedBodyAsHelper = true;
+						break;
+					}
 				}
 
 				if (!bNeedBodyAsHelper)
