Control Rig Physics: Clarify naming (BonesToUse -> BoneMask). UI only.

#Synced-CL 50439659

diff --git a/Engine/Plugins/Experimental/ControlRigPhysics/Source/ControlRigPhysics/Public/RigPhysicsExecution.h b/Engine/Plugins/Experimental/ControlRigPhysics/Source/ControlRigPhysics/Public/RigPhysicsExecution.h
--- a/Engine/Plugins/Experimental/ControlRigPhysics/Source/ControlRigPhysics/Public/RigPhysicsExecution.h
+++ b/Engine/Plugins/Experimental/ControlRigPhysics/Source/ControlRigPhysics/Public/RigPhysicsExecution.h
@@ -373,7 +373,7 @@ struct FRigUnit_HierarchyImportCollisionFromPhysicsAsset : public FRigUnit_Physi
 
 	// If this is empty, then all bones with bodies in the physics asset will be created. Otherwise
 	// only bodies that relate to the specified bones will be created.
-	UPROPERTY(meta = (Input))
+	UPROPERTY(meta = (Input, DisplayName = "Bone Mask"))
 	TArray<FName> BonesToUse;
 
 	// Prefix to the bone names
@@ -430,7 +430,7 @@ struct FRigUnit_HierarchyInstantiateFromPhysicsAsset : public FRigUnit_PhysicsBa
 
 	// If this is empty, then all bodies in the physics asset that match a bone in the hierarchy
 	// will be created. Otherwise only bodies that relate to the specified bones will be created.
-	UPROPERTY(meta = (Input))
+	UPROPERTY(meta = (Input, DisplayName = "Bone Mask"))
 	TArray<FRigElementKey> BonesToUse;
 
 	// Whether to enable the joints authored in the physics asset. Note that you can't have drives
@@ -453,9 +453,9 @@ struct FRigUnit_HierarchyInstantiateFromPhysicsAsset : public FRigUnit_PhysicsBa
 	UPROPERTY(meta = (Input))
 	bool bAddParentSpaceControl = false;
 
-	// If true, and BonesToUse is being used to create a subset of the physics asset, then helper
+	// If true, and BonesToUse/BoneMask is being used to create a subset of the physics asset, then helper
 	// bodies will be created when needed: Kinematic bodies on bones that are connected by a physics
-	// joint to those in BonesToUse. This is used so that (a) instantiating individual chains can be
+	// joint to those in BonesToUse/BoneMask. This is used so that (a) instantiating individual chains can be
 	// done without explicitly including the parent (in the physics joint sense) body, and (b)
 	// entire physics assets can be made from multiple calls to instantiate individual chains.
 	UPROPERTY(meta = (Input))
