Use a transationally safe critical section for control rig physics

#Synced-CL 50487443

diff --git a/Engine/Plugins/Experimental/ControlRigPhysics/Source/ControlRigPhysics/Private/RigPhysicsSimulation.cpp b/Engine/Plugins/Experimental/ControlRigPhysics/Source/ControlRigPhysics/Private/RigPhysicsSimulation.cpp
--- a/Engine/Plugins/Experimental/ControlRigPhysics/Source/ControlRigPhysics/Private/RigPhysicsSimulation.cpp
+++ b/Engine/Plugins/Experimental/ControlRigPhysics/Source/ControlRigPhysics/Private/RigPhysicsSimulation.cpp
@@ -71,7 +71,7 @@ FRigPhysicsSimulation::FRigPhysicsSimulation(const FName InOwnerName)
 	: FRigPhysicsSimulationBase(FRigPhysicsSimulation::StaticStruct()), OwnerName(InOwnerName)
 {
 	WorldObjects = MakeShared<TMap<uint32, FWorldObject>>();
-	SimulationMutex = MakeShared<FCriticalSection>();
+	SimulationMutex = MakeShared<FTransactionallySafeCriticalSection>();
 }
 
 //======================================================================================================================
diff --git a/Engine/Plugins/Experimental/ControlRigPhysics/Source/ControlRigPhysics/Private/RigPhysicsSimulation.h b/Engine/Plugins/Experimental/ControlRigPhysics/Source/ControlRigPhysics/Private/RigPhysicsSimulation.h
--- a/Engine/Plugins/Experimental/ControlRigPhysics/Source/ControlRigPhysics/Private/RigPhysicsSimulation.h
+++ b/Engine/Plugins/Experimental/ControlRigPhysics/Source/ControlRigPhysics/Private/RigPhysicsSimulation.h
@@ -478,6 +478,6 @@ private:
 	// Guard against access to the low level simulation. Note that this access may come from an
 	// async/spawned task, so the mutex itself needs to be shared, as we (as primary owners) may get
 	// deleted whilst the spawned task is in flight.
-	TSharedPtr<FCriticalSection> SimulationMutex;
+	TSharedPtr<FTransactionallySafeCriticalSection> SimulationMutex;
 };
 
