Speculative build fix - use the correct #include

#Synced-CL 50487443

diff --git a/Engine/Plugins/Experimental/ControlRigPhysics/Source/ControlRigPhysics/Private/RigPhysicsSimulation.h b/Engine/Plugins/Experimental/ControlRigPhysics/Source/ControlRigPhysics/Private/RigPhysicsSimulation.h
--- a/Engine/Plugins/Experimental/ControlRigPhysics/Source/ControlRigPhysics/Private/RigPhysicsSimulation.h
+++ b/Engine/Plugins/Experimental/ControlRigPhysics/Source/ControlRigPhysics/Private/RigPhysicsSimulation.h
@@ -12,7 +12,7 @@
 #include "Physics/ImmediatePhysics/ImmediatePhysicsDeclares.h"
 #include "Physics/ImmediatePhysics/ImmediatePhysicsSimulation.h"
 #include "Chaos/PBDJointConstraintTypes.h"
-#include "HAL/CriticalSection.h"
+#include "Misc/TransactionallySafeCriticalSection.h"
 
 #include "RigPhysicsSimulation.generated.h"
 
