Fix for build problems with FTransactionallySafeCriticalSection

#Synced-CL 50487443

diff --git a/Engine/Plugins/Experimental/ControlRigPhysics/Source/ControlRigPhysics/Private/RigPhysicsSimulation.cpp b/Engine/Plugins/Experimental/ControlRigPhysics/Source/ControlRigPhysics/Private/RigPhysicsSimulation.cpp
--- a/Engine/Plugins/Experimental/ControlRigPhysics/Source/ControlRigPhysics/Private/RigPhysicsSimulation.cpp
+++ b/Engine/Plugins/Experimental/ControlRigPhysics/Source/ControlRigPhysics/Private/RigPhysicsSimulation.cpp
@@ -630,7 +630,7 @@ void FRigPhysicsSimulation::Instantiate(
 	}
 
 	// Lock against access to the simulation, in case of access by the WorldObject task
-	FScopeLock Lock(SimulationMutex.Get());
+	UE::TScopeLock Lock(*SimulationMutex.Get());
 
 	// We need the simulation space in order to instantiate properly. This is not ideal, as we may
 	// end up updating the simulation space data twice (thus inserting data into the history). This
diff --git a/Engine/Plugins/Experimental/ControlRigPhysics/Source/ControlRigPhysics/Private/RigPhysicsSimulationWorldObjects.cpp b/Engine/Plugins/Experimental/ControlRigPhysics/Source/ControlRigPhysics/Private/RigPhysicsSimulationWorldObjects.cpp
--- a/Engine/Plugins/Experimental/ControlRigPhysics/Source/ControlRigPhysics/Private/RigPhysicsSimulationWorldObjects.cpp
+++ b/Engine/Plugins/Experimental/ControlRigPhysics/Source/ControlRigPhysics/Private/RigPhysicsSimulationWorldObjects.cpp
@@ -13,6 +13,8 @@
 #include "CollisionQueryParams.h"
 #include "CollisionShape.h"
 
+#include "Misc/ScopeLock.h"
+
 //======================================================================================================================
 // Note that the simulation mutex will already be taken when we are called
 void FRigPhysicsSimulation::UpdateWorldObjectsPrePhysics(const FRigPhysicsSolverSettings& SolverSettings)
@@ -167,7 +169,7 @@ void FRigPhysicsSimulation::UpdateWorldObjectsPostPhysics(
 			{
 				if (TSharedPtr<FCriticalSection> SimulationMutex = WeakSimulationMutex.Pin())
 				{
-					FScopeLock Lock(SimulationMutex.Get());
+					UE::TScopeLock Lock(*SimulationMutex.Get());
 
 					if (TStrongObjectPtr<const UWorld> World = WeakWorld.Pin())
 					{
