Put FTransactionallySafeCriticalSection usage in control rig physics back to previous state to get builds working again

#Synced-CL 50487443

diff --git a/Engine/Plugins/Experimental/ControlRigPhysics/Source/ControlRigPhysics/Private/RigPhysicsSimulation.cpp b/Engine/Plugins/Experimental/ControlRigPhysics/Source/ControlRigPhysics/Private/RigPhysicsSimulation.cpp
--- a/Engine/Plugins/Experimental/ControlRigPhysics/Source/ControlRigPhysics/Private/RigPhysicsSimulation.cpp
+++ b/Engine/Plugins/Experimental/ControlRigPhysics/Source/ControlRigPhysics/Private/RigPhysicsSimulation.cpp
@@ -71,7 +71,7 @@ FRigPhysicsSimulation::FRigPhysicsSimulation(const FName InOwnerName)
 	: FRigPhysicsSimulationBase(FRigPhysicsSimulation::StaticStruct()), OwnerName(InOwnerName)
 {
 	WorldObjects = MakeShared<TMap<uint32, FWorldObject>>();
-	SimulationMutex = MakeShared<FTransactionallySafeCriticalSection>();
+	SimulationMutex = MakeShared<FCriticalSection>();
 }
 
 //======================================================================================================================
@@ -630,7 +630,7 @@ void FRigPhysicsSimulation::Instantiate(
 	}
 
 	// Lock against access to the simulation, in case of access by the WorldObject task
-	UE::TScopeLock Lock(*SimulationMutex.Get());
+	FScopeLock Lock(SimulationMutex.Get());
 
 	// We need the simulation space in order to instantiate properly. This is not ideal, as we may
 	// end up updating the simulation space data twice (thus inserting data into the history). This
diff --git a/Engine/Plugins/Experimental/ControlRigPhysics/Source/ControlRigPhysics/Private/RigPhysicsSimulation.h b/Engine/Plugins/Experimental/ControlRigPhysics/Source/ControlRigPhysics/Private/RigPhysicsSimulation.h
--- a/Engine/Plugins/Experimental/ControlRigPhysics/Source/ControlRigPhysics/Private/RigPhysicsSimulation.h
+++ b/Engine/Plugins/Experimental/ControlRigPhysics/Source/ControlRigPhysics/Private/RigPhysicsSimulation.h
@@ -12,7 +12,7 @@
 #include "Physics/ImmediatePhysics/ImmediatePhysicsDeclares.h"
 #include "Physics/ImmediatePhysics/ImmediatePhysicsSimulation.h"
 #include "Chaos/PBDJointConstraintTypes.h"
-#include "Misc/TransactionallySafeCriticalSection.h"
+#include "HAL/CriticalSection.h"
 
 #include "RigPhysicsSimulation.generated.h"
 
@@ -478,6 +478,6 @@ private:
 	// Guard against access to the low level simulation. Note that this access may come from an
 	// async/spawned task, so the mutex itself needs to be shared, as we (as primary owners) may get
 	// deleted whilst the spawned task is in flight.
-	TSharedPtr<FTransactionallySafeCriticalSection> SimulationMutex;
+	TSharedPtr<FCriticalSection> SimulationMutex;
 };
 
diff --git a/Engine/Plugins/Experimental/ControlRigPhysics/Source/ControlRigPhysics/Private/RigPhysicsSimulationWorldObjects.cpp b/Engine/Plugins/Experimental/ControlRigPhysics/Source/ControlRigPhysics/Private/RigPhysicsSimulationWorldObjects.cpp
--- a/Engine/Plugins/Experimental/ControlRigPhysics/Source/ControlRigPhysics/Private/RigPhysicsSimulationWorldObjects.cpp
+++ b/Engine/Plugins/Experimental/ControlRigPhysics/Source/ControlRigPhysics/Private/RigPhysicsSimulationWorldObjects.cpp
@@ -13,8 +13,6 @@
 #include "CollisionQueryParams.h"
 #include "CollisionShape.h"
 
-#include "Misc/ScopeLock.h"
-
 //======================================================================================================================
 // Note that the simulation mutex will already be taken when we are called
 void FRigPhysicsSimulation::UpdateWorldObjectsPrePhysics(const FRigPhysicsSolverSettings& SolverSettings)
@@ -169,7 +167,7 @@ void FRigPhysicsSimulation::UpdateWorldObjectsPostPhysics(
 			{
 				if (TSharedPtr<FCriticalSection> SimulationMutex = WeakSimulationMutex.Pin())
 				{
-					UE::TScopeLock Lock(*SimulationMutex.Get());
+					FScopeLock Lock(SimulationMutex.Get());
 
 					if (TStrongObjectPtr<const UWorld> World = WeakWorld.Pin())
 					{
