Use FTransactionallySafeCriticalSection instead of FCriticalSection for ControlRigPhysics

#Synced-CL 50532608

diff --git a/Engine/Plugins/Experimental/ControlRigPhysics/Source/ControlRigPhysics/Private/RigPhysicsSimulation.cpp b/Engine/Plugins/Experimental/ControlRigPhysics/Source/ControlRigPhysics/Private/RigPhysicsSimulation.cpp
--- a/Engine/Plugins/Experimental/ControlRigPhysics/Source/ControlRigPhysics/Private/RigPhysicsSimulation.cpp
+++ b/Engine/Plugins/Experimental/ControlRigPhysics/Source/ControlRigPhysics/Private/RigPhysicsSimulation.cpp
@@ -71,7 +71,7 @@ FRigPhysicsSimulation::FRigPhysicsSimulation(const FName InOwnerName)
 	: FRigPhysicsSimulationBase(FRigPhysicsSimulation::StaticStruct()), OwnerName(InOwnerName)
 {
 	WorldObjects = MakeShared<TMap<uint32, FWorldObject>>();
-	SimulationMutex = MakeShared<FCriticalSection>();
+	SimulationMutex = MakeShared<FTransactionallySafeCriticalSection>();
 }
 
 //======================================================================================================================
@@ -630,7 +630,7 @@ void FRigPhysicsSimulation::Instantiate(
 	}
 
 	// Lock against access to the simulation, in case of access by the WorldObject task
-	FScopeLock Lock(SimulationMutex.Get());
+	UE::TScopeLock<FTransactionallySafeCriticalSection> Lock(*SimulationMutex.Get());
 
 	// We need the simulation space in order to instantiate properly. This is not ideal, as we may
 	// end up updating the simulation space data twice (thus inserting data into the history). This
@@ -1803,7 +1803,7 @@ void FRigPhysicsSimulation::StepSimulation(
 	}
 
 	// Lock against access to the simulation, in case of access by the WorldObject task
-	FScopeLock Lock(SimulationMutex.Get());
+	UE::TScopeLock<FTransactionallySafeCriticalSection> Lock(*SimulationMutex.Get());
 
 	// We need to know about the simulation space etc before we can instantiate anything into the right place
 	SimulationSpaceData = UpdateSimulationSpaceStateAndCalculateData(
diff --git a/Engine/Plugins/Experimental/ControlRigPhysics/Source/ControlRigPhysics/Private/RigPhysicsSimulation.h b/Engine/Plugins/Experimental/ControlRigPhysics/Source/ControlRigPhysics/Private/RigPhysicsSimulation.h
--- a/Engine/Plugins/Experimental/ControlRigPhysics/Source/ControlRigPhysics/Private/RigPhysicsSimulation.h
+++ b/Engine/Plugins/Experimental/ControlRigPhysics/Source/ControlRigPhysics/Private/RigPhysicsSimulation.h
@@ -12,7 +12,7 @@
 #include "Physics/ImmediatePhysics/ImmediatePhysicsDeclares.h"
 #include "Physics/ImmediatePhysics/ImmediatePhysicsSimulation.h"
 #include "Chaos/PBDJointConstraintTypes.h"
-#include "HAL/CriticalSection.h"
+#include "Misc/TransactionallySafeCriticalSection.h"
 
 #include "RigPhysicsSimulation.generated.h"
 
@@ -478,6 +478,6 @@ private:
 	// Guard against access to the low level simulation. Note that this access may come from an
 	// async/spawned task, so the mutex itself needs to be shared, as we (as primary owners) may get
 	// deleted whilst the spawned task is in flight.
-	TSharedPtr<FCriticalSection> SimulationMutex;
+	TSharedPtr<FTransactionallySafeCriticalSection> SimulationMutex;
 };
 
diff --git a/Engine/Plugins/Experimental/ControlRigPhysics/Source/ControlRigPhysics/Private/RigPhysicsSimulationVisualization.cpp b/Engine/Plugins/Experimental/ControlRigPhysics/Source/ControlRigPhysics/Private/RigPhysicsSimulationVisualization.cpp
--- a/Engine/Plugins/Experimental/ControlRigPhysics/Source/ControlRigPhysics/Private/RigPhysicsSimulationVisualization.cpp
+++ b/Engine/Plugins/Experimental/ControlRigPhysics/Source/ControlRigPhysics/Private/RigPhysicsSimulationVisualization.cpp
@@ -384,7 +384,7 @@ void FRigPhysicsSimulation::Draw(
 	QUICK_SCOPE_CYCLE_COUNTER(STAT_RigPhysics_Draw);
 
 	// Lock against access to the simulation, in case of access by the WorldObject task
-	FScopeLock Lock(SimulationMutex.Get());
+	UE::TScopeLock<FTransactionallySafeCriticalSection> Lock(*SimulationMutex.Get());
 
 #if CHAOS_DEBUG_DRAW
 	// Danny TODO enable/disable chaos debug draw with a cvar or flag
diff --git a/Engine/Plugins/Experimental/ControlRigPhysics/Source/ControlRigPhysics/Private/RigPhysicsSimulationWorldObjects.cpp b/Engine/Plugins/Experimental/ControlRigPhysics/Source/ControlRigPhysics/Private/RigPhysicsSimulationWorldObjects.cpp
--- a/Engine/Plugins/Experimental/ControlRigPhysics/Source/ControlRigPhysics/Private/RigPhysicsSimulationWorldObjects.cpp
+++ b/Engine/Plugins/Experimental/ControlRigPhysics/Source/ControlRigPhysics/Private/RigPhysicsSimulationWorldObjects.cpp
@@ -125,7 +125,7 @@ void FRigPhysicsSimulation::UpdateWorldObjectsPostPhysics(
 		return;
 	}
 
-	TWeakPtr<FCriticalSection> WeakSimulationMutex = SimulationMutex;
+	TWeakPtr<FTransactionallySafeCriticalSection> WeakSimulationMutex = SimulationMutex;
 
 	TWeakPtr<TMap<uint32, FWorldObject>> WeakWorldObjects = WorldObjects;
 	int64 UpdateCounterCopy = UpdateCounter;
@@ -165,9 +165,9 @@ void FRigPhysicsSimulation::UpdateWorldObjectsPostPhysics(
 			[WeakSimulationMutex, WeakWorld, WeakWorldObjects, OverlapAABox, OverlapAABoxCenter, QueryParams, 
 			ObjectQueryParams, UpdateCounterCopy, OwningActorPtrCopy, SimulationCopy]()
 			{
-				if (TSharedPtr<FCriticalSection> SimulationMutex = WeakSimulationMutex.Pin())
+				if (TSharedPtr<FTransactionallySafeCriticalSection> SimulationMutex = WeakSimulationMutex.Pin())
 				{
-					FScopeLock Lock(SimulationMutex.Get());
+					UE::TScopeLock<FTransactionallySafeCriticalSection> Lock(*SimulationMutex.Get());
 
 					if (TStrongObjectPtr<const UWorld> World = WeakWorld.Pin())
 					{
