Better help and defaults for control rig physics - only has an effect if a variable timestep is being used.

#Synced-CL 51156840

diff --git a/Engine/Plugins/Experimental/ControlRigPhysics/Source/ControlRigPhysics/Public/RigPhysicsData.h b/Engine/Plugins/Experimental/ControlRigPhysics/Source/ControlRigPhysics/Public/RigPhysicsData.h
--- a/Engine/Plugins/Experimental/ControlRigPhysics/Source/ControlRigPhysics/Public/RigPhysicsData.h
+++ b/Engine/Plugins/Experimental/ControlRigPhysics/Source/ControlRigPhysics/Public/RigPhysicsData.h
@@ -510,7 +510,9 @@ struct FRigPhysicsSolverSettings
 	float MaxDepenetrationVelocity = 0.0f;
 
 	/**
-	 * The recommended fixed timestep for the RBAN solver. Set to 0 to run with variable timestep.
+	 * The recommended fixed timestep for the solver. Set to 0 to run with variable timestep, and 
+	 * set MaxDeltaTime to the desired sub-step time if you want to take sub-steps.
+	 * 
 	 * NOTE: If this value is non-zero and less than the current frame time, the simulation will step multiple times
 	 * which increases the cost.
 	 */
@@ -523,10 +525,10 @@ struct FRigPhysicsSolverSettings
 
 	/**
 	 * If a variable timestep is in use, this is the maximum delta time that can be used. If this is smaller than the
-	 * requested total delta time then multiple steps will be used.
+	 * requested total delta time then multiple steps will be used. 
 	 */
 	UPROPERTY(EditAnywhere, Category = SolverSettings, meta = (ClampMin = 0))
-	float MaxDeltaTime = 0.02f;
+	float MaxDeltaTime = 1.0f;
 
 	/**
 	 * Whether to use the linear or non-linear solver for RBAN Joints. The linear solver is significantly cheaper than
