Control Rig Physics - miscellaneous fixes (use of MaxForce/MaxTorque) and tidy-ups from review. Improves/fixes some automatic collision shape generation.

#Synced-CL 51351389

diff --git a/Engine/Plugins/Experimental/ControlRigPhysics/Source/ControlRigPhysics/Private/RigPhysicsBodyComponent.cpp b/Engine/Plugins/Experimental/ControlRigPhysics/Source/ControlRigPhysics/Private/RigPhysicsBodyComponent.cpp
--- a/Engine/Plugins/Experimental/ControlRigPhysics/Source/ControlRigPhysics/Private/RigPhysicsBodyComponent.cpp
+++ b/Engine/Plugins/Experimental/ControlRigPhysics/Source/ControlRigPhysics/Private/RigPhysicsBodyComponent.cpp
@@ -72,13 +72,13 @@ const FSlateIcon& FRigPhysicsBodyComponent::GetIconForUI() const
 	static const FSlateIcon BodyMultipleDefaultIcon = FSlateIcon(
 		FControlRigPhysicsEditorStyle::Get().GetStyleSetName(), "ControlRigPhysics.Component.BodyMultipleDefault");
 	static const FSlateIcon BodyMultipleKinematicIcon = FSlateIcon(
-FControlRigPhysicsEditorStyle::Get().GetStyleSetName(), "ControlRigPhysics.Component.BodyMultipleKinematic");
+		FControlRigPhysicsEditorStyle::Get().GetStyleSetName(), "ControlRigPhysics.Component.BodyMultipleKinematic");
 	static const FSlateIcon BodyMultipleSimulatedIcon = FSlateIcon(
 		FControlRigPhysicsEditorStyle::Get().GetStyleSetName(), "ControlRigPhysics.Component.BodyMultipleSimulated");
 	static const FSlateIcon BodySingleDefaultIcon = FSlateIcon(
 		FControlRigPhysicsEditorStyle::Get().GetStyleSetName(), "ControlRigPhysics.Component.BodySingleDefault");
 	static const FSlateIcon BodySingleKinematicIcon = FSlateIcon(
-FControlRigPhysicsEditorStyle::Get().GetStyleSetName(), "ControlRigPhysics.Component.BodySingleKinematic");
+		FControlRigPhysicsEditorStyle::Get().GetStyleSetName(), "ControlRigPhysics.Component.BodySingleKinematic");
 	static const FSlateIcon BodySingleSimulatedIcon = FSlateIcon(
 		FControlRigPhysicsEditorStyle::Get().GetStyleSetName(), "ControlRigPhysics.Component.BodySingleSimulated");
 
@@ -109,10 +109,11 @@ void FRigPhysicsBodyComponent::OnAddedToHierarchy(URigHierarchy* InHierarchy, UR
 
 //======================================================================================================================
 void FRigPhysicsBodyComponent::AutoCalculateCollision(
-	URigHierarchy* InHierarchy, float MinAspectRatio, float MinSize)
+	URigHierarchy* InHierarchy, float MinAspectRatio, float InMinSize)
 {
 	// Start clean
 	Collision = FRigPhysicsCollision();
+	const float MinSize = FMath::Max(InMinSize, UE_KINDA_SMALL_NUMBER);
 
 	TArray<FVector> Points{ FVector::ZeroVector };
 	FVector MidPoint = FVector::ZeroVector;
@@ -135,21 +136,29 @@ void FRigPhysicsBodyComponent::AutoCalculateCollision(
 		FTransform TM = InHierarchy->GetLocalTransform(GetElementKey(), true);
 		if (TM.GetTranslation().SquaredLength() < UE_SMALL_NUMBER)
 		{
-			// We have no children, and are co-located with our parent. Make a single, arbitrarily
-			// sized, shape.
-			float Size = FMath::Max(MinSize, 10.0f);
-			Collision.Boxes.Add(FRigPhysicsCollisionBox(TM, FVector(Size)));
+			// We have no children, and are co-located with our parent. Make a single, minimally
+			// sized, sphere.
+			Collision.Spheres.Add(FRigPhysicsCollisionSphere(TM, 0.5f * MinSize));
 		}
 		else
 		{
 			TM.SetTranslation(TM.GetTranslation() * 0.5);
-			// We're going to make a box that goes along the X axis
-			TM.SetRotation(FQuat::FindBetweenVectors(FVector(1, 0, 0), TM.GetTranslation()));
-			FVector Extents;
-			Extents.X = FMath::Max(MinSize, TM.GetTranslation().Length() * 2.0);
-			Extents.Y = FMath::Max(MinSize, Extents.X * MinAspectRatio);
-			Extents.Z = Extents.Y;
-			Collision.Boxes.Add(FRigPhysicsCollisionBox(TM, Extents));
+			// We're going to make a capsule, and they are defined as extending along the Z axis
+			TM.SetRotation(FQuat::FindBetweenVectors(FVector(0, 0, 1), TM.GetTranslation()));
+
+			float Length = TM.GetTranslation().Length() * 2.0f;
+			const float Radius = 0.5f * FMath::Max(MinSize, Length * MinAspectRatio);
+			Length -= 2.0f * Radius;
+			Length = FMath::Max(0.0f, Length);
+			if (Length > UE_KINDA_SMALL_NUMBER)
+			{
+				Collision.Capsules.Add(FRigPhysicsCollisionCapsule(TM, Radius, Length));
+			}
+			else
+			{
+				Collision.Spheres.Add(FRigPhysicsCollisionSphere(TM, Radius));
+			}
+
 		}
 	}
 	else if (Points.Num() == 2)
@@ -161,23 +170,29 @@ void FRigPhysicsBodyComponent::AutoCalculateCollision(
 		TM.SetRotation(FQuat::FindBetweenVectors(FVector(0, 0, 1), Points[1]));
 
 		float Length = TM.GetTranslation().Length() * 2.0f;
-		float Radius = FMath::Max(MinSize, Length * MinAspectRatio);
+		const float Radius = 0.5f * FMath::Max(MinSize, Length * MinAspectRatio);
 		Length -= 2.0f * Radius;
-		Length = FMath::Max(0.0f, Length);
-		Collision.Capsules.Add(FRigPhysicsCollisionCapsule(TM, Radius, Length));
+		if (Length > UE_KINDA_SMALL_NUMBER)
+		{
+			Collision.Capsules.Add(FRigPhysicsCollisionCapsule(TM, Radius, Length));
+		}
+		else
+		{
+			Collision.Spheres.Add(FRigPhysicsCollisionSphere(TM, Radius));
+		}
 	}
 	else
 	{
 		// We could calculate the ideal orientation of a box by calculating the eigenvectors of the
 		// covariance matrix that represents all the joint positions relative to the centroid.
-		// However, for now just use a box orientated with the current join - it is simpler, and
+		// However, for now just use a box orientated with the current joint - it is simpler, and
 		// also avoids generating "messy" orientations.
 		UE::Geometry::FAxisAlignedBox3d Box;
 		Box.Contain(Points);
 
-		FTransform TM(Box.Center());
+		const FTransform TM(Box.Center());
 		FVector Extents = Box.Extents() * 2.0; // The returned extents are half the extents!
-		double MaxExtent = Extents.GetAbsMax();
+		const double MaxExtent = Extents.GetAbsMax();
 		Extents.X = FMath::Max(MinSize, FMath::Max(Extents.X, MaxExtent * MinAspectRatio));
 		Extents.Y = FMath::Max(MinSize, FMath::Max(Extents.Y, MaxExtent * MinAspectRatio));
 		Extents.Z = FMath::Max(MinSize, FMath::Max(Extents.Z, MaxExtent * MinAspectRatio));
diff --git a/Engine/Plugins/Experimental/ControlRigPhysics/Source/ControlRigPhysics/Private/RigPhysicsBodyExecution.cpp b/Engine/Plugins/Experimental/ControlRigPhysics/Source/ControlRigPhysics/Private/RigPhysicsBodyExecution.cpp
--- a/Engine/Plugins/Experimental/ControlRigPhysics/Source/ControlRigPhysics/Private/RigPhysicsBodyExecution.cpp
+++ b/Engine/Plugins/Experimental/ControlRigPhysics/Source/ControlRigPhysics/Private/RigPhysicsBodyExecution.cpp
@@ -31,6 +31,7 @@ FRigUnit_AddPhysicsBody_Execute()
 	if (ExecuteContext.GetEventName() != FRigUnit_PrepareForExecution::EventName)
 	{
 		UE_CONTROLRIG_RIGUNIT_REPORT_ERROR(TEXT("AddPhysicsBody can only be used during Setup"));
+		return;
 	}
 
 	if (URigHierarchyController* Controller = ExecuteContext.Hierarchy->GetController())
@@ -67,14 +68,34 @@ FRigUnit_HierarchyAutoCalculateCollision_Execute()
 	if (ExecuteContext.GetEventName() != FRigUnit_PrepareForExecution::EventName)
 	{
 		UE_CONTROLRIG_RIGUNIT_REPORT_ERROR(TEXT("AutoCalculateCollision can only be used during Setup"));
+		return;
 	}
-
 	if (FRigPhysicsBodyComponent* Component = GetBody(ExecuteContext.Hierarchy, PhysicsBodyComponentKey))
 	{
 		Component->AutoCalculateCollision(ExecuteContext.Hierarchy, MinAspectRatio, MinSize);
 	}
 }
 
+//======================================================================================================================
+FRigUnit_HierarchySetBodySolverSettings_Execute()
+{
+	DECLARE_SCOPE_HIERARCHICAL_COUNTER_RIGUNIT();
+	if (!ExecuteContext.Hierarchy)
+	{
+		return;
+	}
+	if (ExecuteContext.GetEventName() != FRigUnit_PrepareForExecution::EventName)
+	{
+		UE_CONTROLRIG_RIGUNIT_REPORT_ERROR(TEXT("SetBodySolverSettings can only be used during Setup"));
+		return;
+	}
+	if (FRigPhysicsBodyComponent* Component = GetBody(ExecuteContext.Hierarchy, PhysicsBodyComponentKey))
+	{
+		Component->BodySolverSettings = BodySolverSettings;
+	}
+}
+
+
 //======================================================================================================================
 FRigUnit_HierarchySetDynamics_Execute()
 {
@@ -83,6 +104,11 @@ FRigUnit_HierarchySetDynamics_Execute()
 	{
 		return;
 	}
+	if (ExecuteContext.GetEventName() != FRigUnit_PrepareForExecution::EventName)
+	{
+		UE_CONTROLRIG_RIGUNIT_REPORT_ERROR(TEXT("SetDynamics can only be used during Setup"));
+		return;
+	}
 	if (FRigPhysicsBodyComponent* Component = GetBody(ExecuteContext.Hierarchy, PhysicsBodyComponentKey))
 	{
 		Component->Dynamics = Dynamics;
@@ -97,6 +123,11 @@ FRigUnit_HierarchySetCollision_Execute()
 	{
 		return;
 	}
+	if (ExecuteContext.GetEventName() != FRigUnit_PrepareForExecution::EventName)
+	{
+		UE_CONTROLRIG_RIGUNIT_REPORT_ERROR(TEXT("SetCollision can only be used during Setup"));
+		return;
+	}
 	if (FRigPhysicsBodyComponent* Component = GetBody(ExecuteContext.Hierarchy, PhysicsBodyComponentKey))
 	{
 		Component->Collision = Collision;
@@ -111,6 +142,11 @@ FRigUnit_HierarchyDisableCollisionBetween_Execute()
 	{
 		return;
 	}
+	if (ExecuteContext.GetEventName() != FRigUnit_PrepareForExecution::EventName)
+	{
+		UE_CONTROLRIG_RIGUNIT_REPORT_ERROR(TEXT("DisableCollisionBetween can only be used during Setup"));
+		return;
+	}
 	if (FRigPhysicsBodyComponent* Component1 = GetBody(ExecuteContext.Hierarchy, PhysicsBodyComponentKey1))
 	{
 		Component1->NoCollisionBodies.Add(PhysicsBodyComponentKey2);
@@ -244,6 +280,7 @@ FRigUnit_HierarchySetPhysicsBodyMaterial_Execute()
 	if (ExecuteContext.GetEventName() != FRigUnit_PrepareForExecution::EventName)
 	{
 		UE_CONTROLRIG_RIGUNIT_REPORT_ERROR(TEXT("SetPhysicsBodyMaterial can only be used during Setup"));
+		return;
 	}
 
 	if (FRigPhysicsBodyComponent* Component = GetBody(ExecuteContext.Hierarchy, PhysicsBodyComponentKey))
diff --git a/Engine/Plugins/Experimental/ControlRigPhysics/Source/ControlRigPhysics/Private/RigPhysicsData.cpp b/Engine/Plugins/Experimental/ControlRigPhysics/Source/ControlRigPhysics/Private/RigPhysicsData.cpp
--- a/Engine/Plugins/Experimental/ControlRigPhysics/Source/ControlRigPhysics/Private/RigPhysicsData.cpp
+++ b/Engine/Plugins/Experimental/ControlRigPhysics/Source/ControlRigPhysics/Private/RigPhysicsData.cpp
@@ -344,9 +344,9 @@ FArchive& operator <<(FArchive& Ar, FRigPhysicsDriveData& Data)
 			Data.LinearDriveConstraint.YDrive.Damping = LinearDamping;
 			Data.LinearDriveConstraint.ZDrive.Damping = LinearDamping;
 
-			Data.LinearDriveConstraint.XDrive.MaxForce = MaxTorque;
-			Data.LinearDriveConstraint.YDrive.MaxForce = MaxTorque;
-			Data.LinearDriveConstraint.ZDrive.MaxForce = MaxTorque;
+			Data.LinearDriveConstraint.XDrive.MaxForce = MaxForce;
+			Data.LinearDriveConstraint.YDrive.MaxForce = MaxForce;
+			Data.LinearDriveConstraint.ZDrive.MaxForce = MaxForce;
 
 			Data.AngularDriveConstraint.SlerpDrive.Stiffness = AngularSpring;
 			Data.AngularDriveConstraint.SwingDrive.Stiffness = AngularSpring;
diff --git a/Engine/Plugins/Experimental/ControlRigPhysics/Source/ControlRigPhysics/Public/RigPhysicsBodyComponent.h b/Engine/Plugins/Experimental/ControlRigPhysics/Source/ControlRigPhysics/Public/RigPhysicsBodyComponent.h
--- a/Engine/Plugins/Experimental/ControlRigPhysics/Source/ControlRigPhysics/Public/RigPhysicsBodyComponent.h
+++ b/Engine/Plugins/Experimental/ControlRigPhysics/Source/ControlRigPhysics/Public/RigPhysicsBodyComponent.h
@@ -24,15 +24,22 @@ public:
 	GENERATED_BODY()
 	DECLARE_RIG_COMPONENT_METHODS(FRigPhysicsBodyComponent)
 
+	// Simulation/solver properties that should only be set during the control rig construction
+	// event, before the body is instantiated
 	UPROPERTY(BlueprintReadOnly, EditAnywhere, Category = Physics, meta = (ShowOnlyInnerProperties))
 	FRigPhysicsBodySolverSettings BodySolverSettings;
 
+	// Dynamics properties that should only be set during the control rig construction event, before
+	// the body is instantiated
 	UPROPERTY(BlueprintReadOnly, EditAnywhere, Category = Physics, meta = (ShowOnlyInnerProperties))
 	FRigPhysicsDynamics Dynamics;
 
+	// Collision properties that should only be set during the control rig construction event,
+	// before the body is instantiated
 	UPROPERTY(BlueprintReadOnly, EditAnywhere, Category = Physics, meta = (ShowOnlyInnerProperties))
 	FRigPhysicsCollision Collision;
 
+	// Properties that can and will likely be modified at runtime
 	UPROPERTY(BlueprintReadOnly, EditAnywhere, Category = Physics)
 	FPhysicsControlModifierData BodyData;
 
@@ -40,14 +47,20 @@ public:
 	UPROPERTY(BlueprintReadOnly, EditAnywhere, Category = Physics)
 	FTransform KinematicTarget;
 
-	// A list of body components with which we should not collide. The solver component can also be included.
+	// A list of body components with which we should not collide. The solver component can also be
+	// included to prevent collisions against its shapes. This is typically/must only be set during
+	// the control rig construction event, before the body is instantiated
 	UPROPERTY(BlueprintReadOnly, EditAnywhere, Category = Physics)
 	TArray<FRigComponentKey> NoCollisionBodies;
 
 	// Removes any existing collision, and replaces it with a shape calculated from the joint
-	// positions (if possible). The shape will be a single box.
-	// @param MinAspectRatio the minimum box extent, as a proportion of the maximum box extent.
-	UE_API void AutoCalculateCollision(URigHierarchy* InHierarchy, float MinAspectRatio = 0.25f, float MinSize = 0.0f);
+	// positions (if possible). The shape will be a single box when there are enough joints, but
+	// when there are only a few joints it may be a sphere or capsule, depending on what is most appropriate.
+	//
+	// @param Hierarchy The hierarchy - the joint positions will be used to calculate the shape.
+	// @param MinAspectRatio The minimum box extent, as a proportion of the maximum box extent.
+	// @param MinSize The minimum size (in each dimension) of the created shape.
+	UE_API void AutoCalculateCollision(URigHierarchy* Hierarchy, float MinAspectRatio = 0.25f, float MinSize = 0.0f);
 
 	UE_API virtual void Save(FArchive& Ar) override;
 	UE_API virtual void Load(FArchive& Ar) override;
diff --git a/Engine/Plugins/Experimental/ControlRigPhysics/Source/ControlRigPhysics/Public/RigPhysicsBodyExecution.h b/Engine/Plugins/Experimental/ControlRigPhysics/Source/ControlRigPhysics/Public/RigPhysicsBodyExecution.h
--- a/Engine/Plugins/Experimental/ControlRigPhysics/Source/ControlRigPhysics/Public/RigPhysicsBodyExecution.h
+++ b/Engine/Plugins/Experimental/ControlRigPhysics/Source/ControlRigPhysics/Public/RigPhysicsBodyExecution.h
@@ -58,7 +58,7 @@ struct FRigUnit_AddPhysicsBody : public FRigUnit_PhysicsBaseMutable
 	FPhysicsControlModifierData BodyData;
 };
 
-// Discards any existing collision data and replaces it with a box based on the joint positions.
+// Discards any existing collision data and replaces it with a shape based on the joint positions.
 // Note that this must be called before the physics solver is instantiated/stepped.
 USTRUCT(meta = (DisplayName = "Calculate Physics Collision", Varying))
 struct FRigUnit_HierarchyAutoCalculateCollision : public FRigUnit_PhysicsBaseMutable
@@ -89,6 +89,31 @@ struct FRigUnit_HierarchyAutoCalculateCollision : public FRigUnit_PhysicsBaseMut
 	float MinSize = 0.0f;
 };
 
+// Sets the body solver settings etc for a physics component body
+USTRUCT(meta = (DisplayName = "Set Physics Body Solver Settings", Varying))
+struct FRigUnit_HierarchySetBodySolverSettings : public FRigUnit_PhysicsBaseMutable
+{
+	GENERATED_BODY()
+
+	FRigUnit_HierarchySetBodySolverSettings()
+	{
+		PhysicsBodyComponentKey.ElementKey.Type = ERigElementType::Bone;
+		PhysicsBodyComponentKey.Name = FRigPhysicsBodyComponent::GetDefaultName();
+	}
+
+	RIGVM_METHOD()
+	UE_API virtual void Execute() override;
+
+	// The Physics Body to be updated
+	UPROPERTY(meta = (Input))
+	FRigComponentKey PhysicsBodyComponentKey;
+
+	// Sets the body solver settings of the Physics Body
+	UPROPERTY(meta = (Input))
+	FRigPhysicsBodySolverSettings BodySolverSettings;
+};
+
+
 // Sets the mass etc for a physics component body
 USTRUCT(meta = (DisplayName = "Set Physics Body Dynamics Properties", Varying))
 struct FRigUnit_HierarchySetDynamics : public FRigUnit_PhysicsBaseMutable
diff --git a/Engine/Plugins/Experimental/ControlRigPhysics/Source/ControlRigPhysics/Public/RigPhysicsData.h b/Engine/Plugins/Experimental/ControlRigPhysics/Source/ControlRigPhysics/Public/RigPhysicsData.h
--- a/Engine/Plugins/Experimental/ControlRigPhysics/Source/ControlRigPhysics/Public/RigPhysicsData.h
+++ b/Engine/Plugins/Experimental/ControlRigPhysics/Source/ControlRigPhysics/Public/RigPhysicsData.h
@@ -32,7 +32,7 @@ struct FRigPhysicsVisualizationSettings
 	GENERATED_BODY()
 
 	// Enable visualization (includes enabling low-level chaos debug draw). Note that this can be
-	// overridden using CVar ControlRig.Physics.VisualizationOveride
+	// overridden using CVar ControlRig.Physics.VisualizationOverride
 	UPROPERTY(EditAnywhere, Category = Visualization)
 	bool bEnableVisualization = true;
 
@@ -65,12 +65,12 @@ struct FRigPhysicsVisualizationSettings
 	bool bShowWorldOverlapBox = false;
 
 	// If visualization is enabled, show active contacts. Note that this can be overridden using
-	// CVar ControlRig.Physics.DrawActiveContactsOveride
+	// CVar ControlRig.Physics.DrawActiveContactsOverride
 	UPROPERTY(EditAnywhere, Category = Visualization)
 	bool bShowActiveContacts = true;
 
 	// If visualization is enabled, show inactive contacts. Note that this can be overridden using
-	// CVar ControlRig.Physics.DrawInactiveContactsOveride
+	// CVar ControlRig.Physics.DrawInactiveContactsOverride
 	UPROPERTY(EditAnywhere, Category = Visualization)
 	bool bShowInactiveContacts = false;
 };
diff --git a/Engine/Plugins/Experimental/PhysicsControl/Source/PhysicsControl/Public/PhysicsControlData.h b/Engine/Plugins/Experimental/PhysicsControl/Source/PhysicsControl/Public/PhysicsControlData.h
--- a/Engine/Plugins/Experimental/PhysicsControl/Source/PhysicsControl/Public/PhysicsControlData.h
+++ b/Engine/Plugins/Experimental/PhysicsControl/Source/PhysicsControl/Public/PhysicsControlData.h
@@ -287,7 +287,9 @@ struct FPhysicsControlModifierData
 	UPROPERTY(EditAnywhere, BlueprintReadWrite, Category = PhysicsControl)
 	TEnumAsByte<ECollisionEnabled::Type> CollisionType = ECollisionEnabled::QueryAndPhysics;
 
-	// The space in which kinematic targets should operate
+	// The space in which kinematic targets should operate. Note that some target spaces are more
+	// expensive than others due to the calculations: IgnoreTarget is the cheapest,
+	// OffsetInBoneSpace and OffsetInComponentSpace are cheap, but OffsetInWorldSpace is relatively expensive.
 	UPROPERTY(BlueprintReadOnly, EditAnywhere, Category = Physics)
 	EPhysicsControlKinematicTargetSpace KinematicTargetSpace = EPhysicsControlKinematicTargetSpace::OffsetInBoneSpace;
 
diff --git a/Engine/Plugins/Experimental/PhysicsControl/Source/PhysicsControl/Public/PhysicsControlHelpers.h b/Engine/Plugins/Experimental/PhysicsControl/Source/PhysicsControl/Public/PhysicsControlHelpers.h
--- a/Engine/Plugins/Experimental/PhysicsControl/Source/PhysicsControl/Public/PhysicsControlHelpers.h
+++ b/Engine/Plugins/Experimental/PhysicsControl/Source/PhysicsControl/Public/PhysicsControlHelpers.h
@@ -48,15 +48,17 @@ PHYSICSCONTROL_API FBodyInstance* GetBodyInstance(UPrimitiveComponent* Component
  */
 PHYSICSCONTROL_API FName GetPhysicalParentBone(USkeletalMeshComponent* SkeletalComponent, FName BoneName);
 
-inline FVector CalculateAngularVelocity(const FQuat& PrevQ, const FQuat& CurrentQ, float Dt)
+// Calculates the angular velocity that would go from PrevQ to CurrentQ in Dt. Returns zero if Dt is invalid.
+FORCEINLINE FVector CalculateAngularVelocity(const FQuat& PrevQ, const FQuat& CurrentQ, float Dt)
 {
 	const FQuat DeltaQ = (CurrentQ * PrevQ.Inverse()).GetShortestArcWith(FQuat::Identity);
-	return DeltaQ.ToRotationVector() / Dt;
+	return (Dt > UE_SMALL_NUMBER) ? DeltaQ.ToRotationVector() / Dt : FVector::ZeroVector;
 }
 
-inline FVector CalculateLinearVelocity(const FVector& PrevP, const FVector& CurrentP, float Dt)
+// Calculates the linear velocity that would go from PrevP to CurrentP in Dt. Returns zero if Dt is invalid.
+FORCEINLINE FVector CalculateLinearVelocity(const FVector& PrevP, const FVector& CurrentP, float Dt)
 {
-	return (CurrentP - PrevP) / Dt;
+	return (Dt > UE_SMALL_NUMBER) ? (CurrentP - PrevP) / Dt : FVector::ZeroVector;
 }
 
 /**
diff --git a/Engine/Plugins/Experimental/PhysicsControl/Source/PhysicsControl/Public/PhysicsControlPoseData.h b/Engine/Plugins/Experimental/PhysicsControl/Source/PhysicsControl/Public/PhysicsControlPoseData.h
--- a/Engine/Plugins/Experimental/PhysicsControl/Source/PhysicsControl/Public/PhysicsControlPoseData.h
+++ b/Engine/Plugins/Experimental/PhysicsControl/Source/PhysicsControl/Public/PhysicsControlPoseData.h
@@ -35,6 +35,9 @@ struct FPosQuat
 	FVector GetTranslation() const { return Translation; }
 	FQuat GetRotation() const { return Rotation; }
 
+	void SetTranslation(const FVector& T) { Translation = T; }
+	void SetRotation(const FQuat& Q) { Rotation = Q; }
+
 	FTransform ToTransform() const
 	{
 		return FTransform(Rotation, Translation);
