Control Rig Physics optimisations. Removes the use of FPosQuat from CRP.

#Synced-CL 51351389

diff --git a/Engine/Plugins/Experimental/ControlRigPhysics/Source/ControlRigPhysics/Private/RigPhysicsControlComponent.cpp b/Engine/Plugins/Experimental/ControlRigPhysics/Source/ControlRigPhysics/Private/RigPhysicsControlComponent.cpp
--- a/Engine/Plugins/Experimental/ControlRigPhysics/Source/ControlRigPhysics/Private/RigPhysicsControlComponent.cpp
+++ b/Engine/Plugins/Experimental/ControlRigPhysics/Source/ControlRigPhysics/Private/RigPhysicsControlComponent.cpp
@@ -41,6 +41,7 @@ void FRigPhysicsControlComponent::Load(FArchive& Ar)
 	{
 		Ar << ControlMultiplier;
 	}
+	DirtyFlags = ERigPhysicsControlComponentDirtyFlags::All;
 }
 
 //======================================================================================================================
@@ -75,4 +76,5 @@ void FRigPhysicsControlComponent::OnRigHierarchyKeyChanged(
 			ChildBodyComponentKey = InNewKey.GetComponent();
 		}
 	}
+	DirtyFlags = ERigPhysicsControlComponentDirtyFlags::All;
 }
diff --git a/Engine/Plugins/Experimental/ControlRigPhysics/Source/ControlRigPhysics/Private/RigPhysicsControlExecution.cpp b/Engine/Plugins/Experimental/ControlRigPhysics/Source/ControlRigPhysics/Private/RigPhysicsControlExecution.cpp
--- a/Engine/Plugins/Experimental/ControlRigPhysics/Source/ControlRigPhysics/Private/RigPhysicsControlExecution.cpp
+++ b/Engine/Plugins/Experimental/ControlRigPhysics/Source/ControlRigPhysics/Private/RigPhysicsControlExecution.cpp
@@ -29,6 +29,7 @@ FRigUnit_AddPhysicsControl_Execute()
 	if (ExecuteContext.GetEventName() != FRigUnit_PrepareForExecution::EventName)
 	{
 		UE_CONTROLRIG_RIGUNIT_REPORT_ERROR(TEXT("AddPhysicsControl can only be used during Setup"));
+		return;
 	}
 
 	if (URigHierarchyController* Controller = ExecuteContext.Hierarchy->GetController())
@@ -49,6 +50,7 @@ FRigUnit_AddPhysicsControl_Execute()
 				Component->ControlData = ControlData;
 				Component->ControlMultiplier = ControlMultiplier;
 				Component->ControlTarget = ControlTarget;
+				Component->DirtyFlags |= ERigPhysicsControlComponentDirtyFlags::Data;
 			}
 		}
 	}
@@ -63,6 +65,7 @@ FRigUnit_HierarchySetControlEnabled_Execute()
 		if (FRigPhysicsControlComponent* Component = GetControl(ExecuteContext.Hierarchy, PhysicsControlComponentKey))
 		{
 			Component->ControlData.bEnabled = bEnabled;
+			Component->DirtyFlags |= ERigPhysicsControlComponentDirtyFlags::Data;
 		}
 	}
 }
@@ -75,13 +78,17 @@ FRigUnit_HierarchySetControlCustomControlPoint_Execute()
 	{
 		if (FRigPhysicsControlComponent* Component = GetControl(ExecuteContext.Hierarchy, PhysicsControlComponentKey))
 		{
-			Component->ControlData.CustomControlPoint = CustomControlPoint;
-			Component->ControlData.bUseCustomControlPoint = bUseCustomControlPoint;
+			if (Component->ControlData.CustomControlPoint != CustomControlPoint ||
+				Component->ControlData.bUseCustomControlPoint != bUseCustomControlPoint)
+			{
+				Component->ControlData.CustomControlPoint = CustomControlPoint;
+				Component->ControlData.bUseCustomControlPoint = bUseCustomControlPoint;
+				Component->DirtyFlags |= ERigPhysicsControlComponentDirtyFlags::Data;
+			}
 		}
 	}
 }
 
-
 //======================================================================================================================
 FRigUnit_HierarchySetControlData_Execute()
 {
@@ -90,86 +97,99 @@ FRigUnit_HierarchySetControlData_Execute()
 	{
 		if (FRigPhysicsControlComponent* Component = GetControl(ExecuteContext.Hierarchy, PhysicsControlComponentKey))
 		{
-			Component->ControlData = ControlData;
+			if (Component->ControlData != ControlData)
+			{
+				Component->ControlData = ControlData;
+				Component->DirtyFlags |= ERigPhysicsControlComponentDirtyFlags::Data;
+			}
 		}
 	}
 }
 
 //======================================================================================================================
-FRigUnit_HierarchySetControlLinearStrength_Execute()
+FRigUnit_HierarchyGetControlData_Execute()
 {
 	DECLARE_SCOPE_HIERARCHICAL_COUNTER_RIGUNIT();
 	if (ExecuteContext.Hierarchy)
 	{
-		if (FRigPhysicsControlComponent* Component = GetControl(ExecuteContext.Hierarchy, PhysicsControlComponentKey))
+		if (const FRigPhysicsControlComponent* Component =
+			GetControl(ExecuteContext.Hierarchy, PhysicsControlComponentKey))
 		{
-			Component->ControlData.LinearStrength = Strength;
+			ControlData = Component->ControlData;
 		}
 	}
 }
 
 //======================================================================================================================
-FRigUnit_HierarchySetControlLinearDampingRatio_Execute()
+FRigUnit_HierarchySetControlLinearStrength_Execute()
 {
 	DECLARE_SCOPE_HIERARCHICAL_COUNTER_RIGUNIT();
 	if (ExecuteContext.Hierarchy)
 	{
 		if (FRigPhysicsControlComponent* Component = GetControl(ExecuteContext.Hierarchy, PhysicsControlComponentKey))
 		{
-			Component->ControlData.LinearDampingRatio = DampingRatio;
+			Component->ControlData.LinearStrength = Strength;
+			Component->DirtyFlags |= ERigPhysicsControlComponentDirtyFlags::Data;
 		}
 	}
 }
 
 //======================================================================================================================
-FRigUnit_HierarchySetControlTargetVelocityMultipliers_Execute()
+FRigUnit_HierarchySetControlLinearDampingRatio_Execute()
 {
 	DECLARE_SCOPE_HIERARCHICAL_COUNTER_RIGUNIT();
 	if (ExecuteContext.Hierarchy)
 	{
 		if (FRigPhysicsControlComponent* Component = GetControl(ExecuteContext.Hierarchy, PhysicsControlComponentKey))
 		{
-			Component->ControlData.LinearTargetVelocityMultiplier = LinearTargetVelocityMultiplier;
-			Component->ControlData.AngularTargetVelocityMultiplier = AngularTargetVelocityMultiplier;
+			Component->ControlData.LinearDampingRatio = DampingRatio;
+			Component->DirtyFlags |= ERigPhysicsControlComponentDirtyFlags::Data;
 		}
 	}
 }
 
 //======================================================================================================================
-FRigUnit_HierarchySetControlAngularStrength_Execute()
+FRigUnit_HierarchySetControlTargetVelocityMultipliers_Execute()
 {
 	DECLARE_SCOPE_HIERARCHICAL_COUNTER_RIGUNIT();
 	if (ExecuteContext.Hierarchy)
 	{
 		if (FRigPhysicsControlComponent* Component = GetControl(ExecuteContext.Hierarchy, PhysicsControlComponentKey))
 		{
-			Component->ControlData.AngularStrength = Strength;
+			Component->ControlData.LinearTargetVelocityMultiplier = LinearTargetVelocityMultiplier;
+			Component->ControlData.AngularTargetVelocityMultiplier = AngularTargetVelocityMultiplier;
+			Component->DirtyFlags |= ERigPhysicsControlComponentDirtyFlags::Data;
 		}
 	}
 }
 
 //======================================================================================================================
-FRigUnit_HierarchySetControlAngularDampingRatio_Execute()
+FRigUnit_HierarchySetControlAngularStrength_Execute()
 {
 	DECLARE_SCOPE_HIERARCHICAL_COUNTER_RIGUNIT();
 	if (ExecuteContext.Hierarchy)
 	{
 		if (FRigPhysicsControlComponent* Component = GetControl(ExecuteContext.Hierarchy, PhysicsControlComponentKey))
 		{
-			Component->ControlData.AngularDampingRatio = DampingRatio;
+			Component->ControlData.AngularStrength = Strength;
+			Component->DirtyFlags |= ERigPhysicsControlComponentDirtyFlags::Data;
 		}
 	}
 }
 
 //======================================================================================================================
-FRigUnit_HierarchyGetControlData_Execute()
+FRigUnit_HierarchySetControlAngularDampingRatio_Execute()
 {
 	DECLARE_SCOPE_HIERARCHICAL_COUNTER_RIGUNIT();
 	if (ExecuteContext.Hierarchy)
 	{
 		if (FRigPhysicsControlComponent* Component = GetControl(ExecuteContext.Hierarchy, PhysicsControlComponentKey))
 		{
-			ControlData = Component->ControlData;
+			if (Component->ControlData.AngularDampingRatio != DampingRatio)
+			{
+				Component->ControlData.AngularDampingRatio = DampingRatio;
+				Component->DirtyFlags |= ERigPhysicsControlComponentDirtyFlags::Data;
+			}
 		}
 	}
 }
@@ -182,7 +202,11 @@ FRigUnit_HierarchySetControlMultiplier_Execute()
 	{
 		if (FRigPhysicsControlComponent* Component = GetControl(ExecuteContext.Hierarchy, PhysicsControlComponentKey))
 		{
-			Component->ControlMultiplier = ControlMultiplier;
+			if (Component->ControlMultiplier != ControlMultiplier)
+			{
+				Component->ControlMultiplier = ControlMultiplier;
+				Component->DirtyFlags |= ERigPhysicsControlComponentDirtyFlags::Data;
+			}
 		}
 	}
 }
@@ -195,8 +219,16 @@ FRigUnit_HierarchySetControlDataAndMultiplier_Execute()
 	{
 		if (FRigPhysicsControlComponent* Component = GetControl(ExecuteContext.Hierarchy, PhysicsControlComponentKey))
 		{
-			Component->ControlData = ControlData;
-			Component->ControlMultiplier = ControlMultiplier;
+			if (Component->ControlData != ControlData)
+			{
+				Component->ControlData = ControlData;
+				Component->DirtyFlags |= ERigPhysicsControlComponentDirtyFlags::Data;
+			}
+			if (Component->ControlMultiplier != ControlMultiplier)
+			{
+				Component->ControlMultiplier = ControlMultiplier;
+				Component->DirtyFlags |= ERigPhysicsControlComponentDirtyFlags::Data;
+			}
 		}
 	}
 }
@@ -210,6 +242,7 @@ FRigUnit_HierarchySetControlTarget_Execute()
 		if (FRigPhysicsControlComponent* Component = GetControl(ExecuteContext.Hierarchy, PhysicsControlComponentKey))
 		{
 			Component->ControlTarget = ControlTarget;
+			Component->DirtyFlags |= ERigPhysicsControlComponentDirtyFlags::Target;
 		}
 	}
 }
@@ -228,6 +261,7 @@ FRigUnit_HierarchyUpdateControlTarget_Execute()
 				TargetOrientation.Quaternion(), Component->ControlTarget.TargetOrientation.Quaternion(), DeltaTime);
 			Component->ControlTarget.TargetPosition = TargetPosition;
 			Component->ControlTarget.TargetOrientation = TargetOrientation;
+			Component->DirtyFlags |= ERigPhysicsControlComponentDirtyFlags::Target;
 		}
 	}
 }
diff --git a/Engine/Plugins/Experimental/ControlRigPhysics/Source/ControlRigPhysics/Private/RigPhysicsExecution.cpp b/Engine/Plugins/Experimental/ControlRigPhysics/Source/ControlRigPhysics/Private/RigPhysicsExecution.cpp
--- a/Engine/Plugins/Experimental/ControlRigPhysics/Source/ControlRigPhysics/Private/RigPhysicsExecution.cpp
+++ b/Engine/Plugins/Experimental/ControlRigPhysics/Source/ControlRigPhysics/Private/RigPhysicsExecution.cpp
@@ -45,6 +45,7 @@ FRigUnit_AddPhysicsSolver_Execute()
 	if (ExecuteContext.GetEventName() != FRigUnit_PrepareForExecution::EventName)
 	{
 		UE_CONTROLRIG_RIGUNIT_REPORT_ERROR(TEXT("AddPhysicsSolver can only be used during Setup"));
+		return;
 	}
 
 	if (URigHierarchyController* Controller = ExecuteContext.Hierarchy->GetController())
@@ -116,7 +117,7 @@ FRigUnit_StepPhysicsSolver_Execute()
 
 				if (ControlRig->RegisterPhysicsSimulation(NewPhysicsSimulation, PhysicsSolverComponentKey))
 				{
-					NewPhysicsSimulation->Instantiate(ExecuteContext, Hierarchy, SolverComponent);
+					NewPhysicsSimulation->Instantiate(ExecuteContext, Hierarchy, *SolverComponent);
 				}
 				PhysicsSimulation = NewPhysicsSimulation.Get();
 			}
@@ -127,7 +128,7 @@ FRigUnit_StepPhysicsSolver_Execute()
 				const AActor* OwningActorPtr = OwningSceneComponent ? OwningSceneComponent->GetOwner() : nullptr;
 
 				PhysicsSimulation->StepSimulation(
-					ControlRig->GetWorld(), OwningActorPtr, ExecuteContext, Hierarchy, SolverComponent, 
+					ControlRig->GetWorld(), OwningActorPtr, ExecuteContext, Hierarchy, *SolverComponent, 
 					DeltaTimeOverride, SimulationSpaceDeltaTimeOverride, Alpha);
 				int32 VisualizationOverride = CVarControlRigPhysicsVisualizationOverride.GetValueOnAnyThread();
 				if (VisualizationOverride == 1 || 
@@ -211,7 +212,7 @@ FRigUnit_InstantiatePhysics_Execute()
 
 			if (ControlRig->RegisterPhysicsSimulation(PhysicsSimulation, PhysicsSolverComponentKey))
 			{
-				PhysicsSimulation->Instantiate(ExecuteContext, Hierarchy, SolverComponent);
+				PhysicsSimulation->Instantiate(ExecuteContext, Hierarchy, *SolverComponent);
 			}
 		}
 	}
@@ -228,6 +229,7 @@ FRigUnit_AddPhysicsComponents_Execute()
 	if (ExecuteContext.GetEventName() != FRigUnit_PrepareForExecution::EventName)
 	{
 		UE_CONTROLRIG_RIGUNIT_REPORT_ERROR(TEXT("AddPhysicsComponents can only be used during Setup"));
+		return;
 	}
 
 	if (URigHierarchyController* Controller = ExecuteContext.Hierarchy->GetController())
@@ -265,6 +267,7 @@ FRigUnit_AddPhysicsComponents_Execute()
 						Component->ChildBodyComponentKey = PhysicsBodyComponentKey;
 						Component->JointData = JointData;
 						Component->DriveData = DriveData;
+						Component->DirtyFlags |= ERigPhysicsJointComponentDirtyFlags::All;
 					}
 				}
 			}
@@ -741,6 +744,7 @@ FRigUnit_HierarchyInstantiateFromPhysicsAsset_Execute()
 					Component->DriveData = DriveData;
 					Component->ParentBodyComponentKey = *ParentBodyComponentKey;
 					Component->ChildBodyComponentKey = *ChildBodyComponentKey;
+					Component->DirtyFlags |= ERigPhysicsJointComponentDirtyFlags::All;
 
 					PhysicsJointComponentKeys.Add(PhysicsJointComponentKey);
 				}
diff --git a/Engine/Plugins/Experimental/ControlRigPhysics/Source/ControlRigPhysics/Private/RigPhysicsJointComponent.cpp b/Engine/Plugins/Experimental/ControlRigPhysics/Source/ControlRigPhysics/Private/RigPhysicsJointComponent.cpp
--- a/Engine/Plugins/Experimental/ControlRigPhysics/Source/ControlRigPhysics/Private/RigPhysicsJointComponent.cpp
+++ b/Engine/Plugins/Experimental/ControlRigPhysics/Source/ControlRigPhysics/Private/RigPhysicsJointComponent.cpp
@@ -33,6 +33,7 @@ void FRigPhysicsJointComponent::Load(FArchive& Ar)
 	Ar << ChildBodyComponentKey;
 	Ar << JointData;
 	Ar << DriveData;
+	DirtyFlags = ERigPhysicsJointComponentDirtyFlags::All;
 }
 
 //======================================================================================================================
@@ -67,4 +68,5 @@ void FRigPhysicsJointComponent::OnRigHierarchyKeyChanged(const FRigHierarchyKey&
 			ChildBodyComponentKey = InNewKey.GetComponent();
 		}
 	}
+	DirtyFlags = ERigPhysicsJointComponentDirtyFlags::All;
 }
diff --git a/Engine/Plugins/Experimental/ControlRigPhysics/Source/ControlRigPhysics/Private/RigPhysicsJointExecution.cpp b/Engine/Plugins/Experimental/ControlRigPhysics/Source/ControlRigPhysics/Private/RigPhysicsJointExecution.cpp
--- a/Engine/Plugins/Experimental/ControlRigPhysics/Source/ControlRigPhysics/Private/RigPhysicsJointExecution.cpp
+++ b/Engine/Plugins/Experimental/ControlRigPhysics/Source/ControlRigPhysics/Private/RigPhysicsJointExecution.cpp
@@ -36,6 +36,7 @@ FRigUnit_AddPhysicsJoint_Execute()
 	if (ExecuteContext.GetEventName() != FRigUnit_PrepareForExecution::EventName)
 	{
 		UE_CONTROLRIG_RIGUNIT_REPORT_ERROR(TEXT("AddPhysicsJoint can only be used during Setup"));
+		return;
 	}
 
 	if (URigHierarchyController* Controller = ExecuteContext.Hierarchy->GetController())
@@ -53,6 +54,7 @@ FRigUnit_AddPhysicsJoint_Execute()
 				Component->ChildBodyComponentKey = ChildBodyComponentKey;
 				Component->JointData = JointData;
 				Component->DriveData = DriveData;
+				Component->DirtyFlags = ERigPhysicsJointComponentDirtyFlags::All;
 			}
 		}
 	}
@@ -69,6 +71,7 @@ FRigUnit_HierarchySetJointData_Execute()
 	if (FRigPhysicsJointComponent* Component = GetJoint(ExecuteContext.Hierarchy, PhysicsJointComponentKey))
 	{
 		Component->JointData = JointData;
+		Component->DirtyFlags |= ERigPhysicsJointComponentDirtyFlags::Joint;
 	}
 }
 
@@ -97,6 +100,7 @@ FRigUnit_HierarchySetJointLinearLimit_Execute()
 	if (FRigPhysicsJointComponent* Component = GetJoint(ExecuteContext.Hierarchy, PhysicsJointComponentKey))
 	{
 		Component->JointData.LinearConstraint = LinearLimit.LinearConstraint;
+		Component->DirtyFlags |= ERigPhysicsJointComponentDirtyFlags::Joint;
 	}
 }
 
@@ -111,6 +115,7 @@ FRigUnit_HierarchySetJointConeLimit_Execute()
 	if (FRigPhysicsJointComponent* Component = GetJoint(ExecuteContext.Hierarchy, PhysicsJointComponentKey))
 	{
 		Component->JointData.ConeConstraint = ConeLimit.ConeConstraint;
+		Component->DirtyFlags |= ERigPhysicsJointComponentDirtyFlags::Joint;
 	}
 }
 
@@ -125,6 +130,7 @@ FRigUnit_HierarchySetJointTwistLimit_Execute()
 	if (FRigPhysicsJointComponent* Component = GetJoint(ExecuteContext.Hierarchy, PhysicsJointComponentKey))
 	{
 		Component->JointData.TwistConstraint = TwistLimit.TwistConstraint;
+		Component->DirtyFlags |= ERigPhysicsJointComponentDirtyFlags::Joint;
 	}
 }
 
@@ -139,6 +145,7 @@ FRigUnit_HierarchySetJointEnabled_Execute()
 	if (FRigPhysicsJointComponent* Component = GetJoint(ExecuteContext.Hierarchy, PhysicsJointComponentKey))
 	{
 		Component->JointData.bEnabled = bEnabled;
+		Component->DirtyFlags |= ERigPhysicsJointComponentDirtyFlags::Joint;
 	}
 }
 
@@ -153,6 +160,7 @@ FRigUnit_HierarchySetJointDriveData_Execute()
 	if (FRigPhysicsJointComponent* Component = GetJoint(ExecuteContext.Hierarchy, PhysicsJointComponentKey))
 	{
 		Component->DriveData = DriveData;
+		Component->DirtyFlags |= ERigPhysicsJointComponentDirtyFlags::Drive;
 	}
 }
 
@@ -181,6 +189,7 @@ FRigUnit_HierarchySetJointDriveUseSkeletalAnimation_Execute()
 	if (FRigPhysicsJointComponent* Component = GetJoint(ExecuteContext.Hierarchy, PhysicsJointComponentKey))
 	{
 		Component->DriveData.bUseSkeletalAnimation = bUseSkeletalAnimation;
+		Component->DirtyFlags |= ERigPhysicsJointComponentDirtyFlags::Drive;
 	}
 }
 
diff --git a/Engine/Plugins/Experimental/ControlRigPhysics/Source/ControlRigPhysics/Private/RigPhysicsSimulation.cpp b/Engine/Plugins/Experimental/ControlRigPhysics/Source/ControlRigPhysics/Private/RigPhysicsSimulation.cpp
--- a/Engine/Plugins/Experimental/ControlRigPhysics/Source/ControlRigPhysics/Private/RigPhysicsSimulation.cpp
+++ b/Engine/Plugins/Experimental/ControlRigPhysics/Source/ControlRigPhysics/Private/RigPhysicsSimulation.cpp
@@ -1,6 +1,8 @@
 // Copyright Epic Games, Inc. All Rights Reserved.
 
 #include "RigPhysicsSimulation.h"
+#include "RigPhysicsSimulationSpace.inl"
+
 #include "RigPhysicsBodyComponent.h"
 #include "RigPhysicsJointComponent.h"
 #include "RigPhysicsControlComponent.h"
@@ -41,6 +43,8 @@
 #include "PhysicsEngine/ConstraintInstance.h"
 #include "PhysicsEngine/BodySetup.h"
 
+#include "Stats/Stats.h"
+
 #if WITH_CHAOS_VISUAL_DEBUGGER
 #include "ChaosVDRuntimeModule.h"
 #endif
@@ -76,25 +80,25 @@ FRigPhysicsSimulation::FRigPhysicsSimulation(const FName InOwnerName)
 
 //======================================================================================================================
 bool FRigPhysicsSimulation::ShouldComponentBeInSimulation(
-	const URigHierarchy& Hierarchy, FRigComponentKey SolverComponentKey, FRigComponentKey ComponentKey) const
+	const URigHierarchy&              Hierarchy, 
+	const FRigPhysicsSolverComponent& SolverComponent, 
+	const FRigComponentKey&           ComponentKey) const
 {
 	const FRigPhysicsBodyComponent* PhysicsBodyComponent = Cast<FRigPhysicsBodyComponent>(
 		Hierarchy.FindComponent(ComponentKey));
 
-	const FRigPhysicsSolverComponent* PhysicsSolverComponent = Cast<FRigPhysicsSolverComponent>(
-		Hierarchy.FindComponent(SolverComponentKey));
-
-	if (!PhysicsBodyComponent || !PhysicsSolverComponent)
+	if (!PhysicsBodyComponent)
 	{
 		return false;
 	}
 
+	const FRigComponentKey SolverComponentKey = SolverComponent.GetKey();
 	if (PhysicsBodyComponent->BodySolverSettings.PhysicsSolverComponentKey == SolverComponentKey)
 	{
 		return true;
 	}
 
-	if (!PhysicsSolverComponent->SolverSettings.bAutomaticallyAddPhysicsBodyComponents)
+	if (!SolverComponent.SolverSettings.bAutomaticallyAddPhysicsBodyComponents)
 	{
 		return false;
 	}
@@ -121,7 +125,7 @@ bool FRigPhysicsSimulation::ShouldComponentBeInSimulation(
 
 
 //======================================================================================================================
-void FRigPhysicsSimulation::InitialiseSimulation(const FRigPhysicsSolverComponent* SolverComponent)
+void FRigPhysicsSimulation::InitialiseSimulation()
 {
 	DestroyPhysicsSimulation();
 
@@ -150,17 +154,10 @@ void FRigPhysicsSimulation::InitialiseSimulation(const FRigPhysicsSolverComponen
 
 //======================================================================================================================
 void FRigPhysicsSimulation::InitialiseControlRecords(
-	const URigHierarchy& Hierarchy, const FRigPhysicsSolverComponent* SolverComponent)
+	const URigHierarchy& Hierarchy, const FRigPhysicsSolverComponent& SolverComponent)
 {
 	ensure(ControlRecords.IsEmpty());
 
-	if (!SolverComponent)
-	{
-		return;
-	}
-
-	FRigComponentKey SolverComponentKey = SolverComponent->GetKey();
-
 	TArray<FRigComponentKey> AllComponentKeys = Hierarchy.GetAllComponentKeys();
 
 	for (FRigComponentKey ComponentKey : AllComponentKeys)
@@ -179,7 +176,7 @@ void FRigPhysicsSimulation::InitialiseControlRecords(
 				TArray<FRigComponentKey> SiblingComponentKeys = Hierarchy.GetComponentKeys(ComponentKey.ElementKey);
 				for (FRigComponentKey SiblingComponentKey : SiblingComponentKeys)
 				{
-					if (ShouldComponentBeInSimulation(Hierarchy, SolverComponentKey, SiblingComponentKey))
+					if (ShouldComponentBeInSimulation(Hierarchy, SolverComponent, SiblingComponentKey))
 					{
 						ControlRecord.ChildBodyComponentKey = SiblingComponentKey;
 						break;
@@ -195,7 +192,7 @@ void FRigPhysicsSimulation::InitialiseControlRecords(
 					TArray<FRigComponentKey> ParentComponentKeys = Hierarchy.GetComponentKeys(ParentElementKey);
 					for (FRigComponentKey ParentComponentKey : ParentComponentKeys)
 					{
-						if (ShouldComponentBeInSimulation(Hierarchy, SolverComponentKey, ParentComponentKey))
+						if (ShouldComponentBeInSimulation(Hierarchy, SolverComponent, ParentComponentKey))
 						{
 							ControlRecord.ParentBodyComponentKey = ParentComponentKey;
 							break;
@@ -204,11 +201,11 @@ void FRigPhysicsSimulation::InitialiseControlRecords(
 				}
 			}
 
-			if (ShouldComponentBeInSimulation(Hierarchy, SolverComponentKey, ControlRecord.ChildBodyComponentKey))
+			if (ShouldComponentBeInSimulation(Hierarchy, SolverComponent, ControlRecord.ChildBodyComponentKey))
 			{
 				// Here, an invalid parent component key indicates a sim-space control. 
 				if (!ControlRecord.ParentBodyComponentKey.IsValid() ||
-					ShouldComponentBeInSimulation(Hierarchy, SolverComponentKey, ControlRecord.ParentBodyComponentKey))
+					ShouldComponentBeInSimulation(Hierarchy, SolverComponent, ControlRecord.ParentBodyComponentKey))
 				{
 					// Just make the record for now - it will be instantiated later
 					ControlRecords.Add(ComponentKey, ControlRecord);
@@ -220,15 +217,11 @@ void FRigPhysicsSimulation::InitialiseControlRecords(
 
 //======================================================================================================================
 void FRigPhysicsSimulation::InitialiseJointRecords(
-	const URigHierarchy& Hierarchy, const FRigPhysicsSolverComponent* SolverComponent)
+	const URigHierarchy& Hierarchy, const FRigPhysicsSolverComponent& SolverComponent)
 {
 	ensure(JointRecords.IsEmpty());
-	if (!SolverComponent)
-	{
-		return;
-	}
 
-	FRigComponentKey SolverComponentKey = SolverComponent->GetKey();
+	FRigComponentKey SolverComponentKey = SolverComponent.GetKey();
 
 	TArray<FRigComponentKey> AllComponentKeys = Hierarchy.GetAllComponentKeys();
 
@@ -247,7 +240,7 @@ void FRigPhysicsSimulation::InitialiseJointRecords(
 				TArray<FRigComponentKey> SiblingComponentKeys = Hierarchy.GetComponentKeys(ComponentKey.ElementKey);
 				for (FRigComponentKey SiblingComponentKey : SiblingComponentKeys)
 				{
-					if (ShouldComponentBeInSimulation(Hierarchy, SolverComponentKey, SiblingComponentKey))
+					if (ShouldComponentBeInSimulation(Hierarchy, SolverComponent, SiblingComponentKey))
 					{
 						JointRecord.ChildBodyComponentKey = SiblingComponentKey;
 						break;
@@ -261,7 +254,7 @@ void FRigPhysicsSimulation::InitialiseJointRecords(
 				TArray<FRigComponentKey> ParentComponentKeys = Hierarchy.GetComponentKeys(ParentElementKey);
 				for (FRigComponentKey ParentComponentKey : ParentComponentKeys)
 				{
-					if (ShouldComponentBeInSimulation(Hierarchy, SolverComponentKey, ParentComponentKey))
+					if (ShouldComponentBeInSimulation(Hierarchy, SolverComponent, ParentComponentKey))
 					{
 						JointRecord.ParentBodyComponentKey = ParentComponentKey;
 						break;
@@ -269,9 +262,9 @@ void FRigPhysicsSimulation::InitialiseJointRecords(
 				}
 			}
 
-			if (ShouldComponentBeInSimulation(Hierarchy, SolverComponentKey, JointRecord.ChildBodyComponentKey))
+			if (ShouldComponentBeInSimulation(Hierarchy, SolverComponent, JointRecord.ChildBodyComponentKey))
 			{
-				if (ShouldComponentBeInSimulation(Hierarchy, SolverComponentKey, JointRecord.ParentBodyComponentKey))
+				if (ShouldComponentBeInSimulation(Hierarchy, SolverComponent, JointRecord.ParentBodyComponentKey))
 				{
 					// Just make the record for now - it will be instantiated later
 					JointRecords.Add(ComponentKey, JointRecord);
@@ -283,15 +276,12 @@ void FRigPhysicsSimulation::InitialiseJointRecords(
 
 
 //======================================================================================================================
-void FRigPhysicsSimulation::InitialiseBodyRecords(const URigHierarchy& Hierarchy, const FRigPhysicsSolverComponent* SolverComponent)
+void FRigPhysicsSimulation::InitialiseBodyRecords(
+	const URigHierarchy& Hierarchy, const FRigPhysicsSolverComponent& SolverComponent)
 {
 	ensure(BodyRecords.IsEmpty());
-	if (!SolverComponent)
-	{
-		return;
-	}
 
-	FRigComponentKey SolverComponentKey = SolverComponent->GetKey();
+	FRigComponentKey SolverComponentKey = SolverComponent.GetKey();
 
 	TArray<FRigComponentKey> AllComponentKeys = Hierarchy.GetAllComponentKeys();
 
@@ -300,7 +290,7 @@ void FRigPhysicsSimulation::InitialiseBodyRecords(const URigHierarchy& Hierarchy
 
 	for (FRigComponentKey ComponentKey : AllComponentKeys)
 	{
-		if (ShouldComponentBeInSimulation(Hierarchy, SolverComponentKey, ComponentKey))
+		if (ShouldComponentBeInSimulation(Hierarchy, SolverComponent, ComponentKey))
 		{
 			// Just make the record for now - it will be instantiated later
 			BodyRecords.Add(ComponentKey);
@@ -326,46 +316,51 @@ void FRigPhysicsSimulation::InitialiseBodyRecords(const URigHierarchy& Hierarchy
 //======================================================================================================================
 void FRigPhysicsSimulation::DestroyPhysicsSimulation()
 {
-	for (TPair<FRigComponentKey, FRigBodyRecord>& BodyRecordPair : BodyRecords)
+	if (Simulation.IsValid())
 	{
-		FRigBodyRecord& Record = BodyRecordPair.Value;
-		if (Record.ActorHandle)
+		for (TPair<FRigComponentKey, FRigBodyRecord>& BodyRecordPair : BodyRecords)
 		{
-			Simulation->DestroyActor(Record.ActorHandle);
+			FRigBodyRecord& Record = BodyRecordPair.Value;
+			if (Record.ActorHandle)
+			{
+				Simulation->DestroyActor(Record.ActorHandle);
+			}
 		}
-	}
-	BodyRecords.Empty();
 
-	for (TPair<FRigComponentKey, FRigJointRecord>& JointRecordPair : JointRecords)
-	{
-		FRigJointRecord& Record = JointRecordPair.Value;
-		if (Record.JointHandle)
+		for (TPair<FRigComponentKey, FRigJointRecord>& JointRecordPair : JointRecords)
 		{
-			Simulation->DestroyJoint(Record.JointHandle);
+			FRigJointRecord& Record = JointRecordPair.Value;
+			if (Record.JointHandle)
+			{
+				Simulation->DestroyJoint(Record.JointHandle);
+			}
 		}
-	}
-	JointRecords.Empty();
 
-	for (TPair<FRigComponentKey, FRigControlRecord>& ControlRecordPair : ControlRecords)
-	{
-		FRigControlRecord& Record = ControlRecordPair.Value;
-		if (Record.JointHandle)
+		for (TPair<FRigComponentKey, FRigControlRecord>& ControlRecordPair : ControlRecords)
+		{
+			FRigControlRecord& Record = ControlRecordPair.Value;
+			if (Record.JointHandle)
+			{
+				Simulation->DestroyJoint(Record.JointHandle);
+			}
+		}
+
+		if (CollisionActorHandle)
 		{
-			Simulation->DestroyJoint(Record.JointHandle);
+			Simulation->DestroyActor(CollisionActorHandle);
+		}
+
+		if (SimulationActorHandle)
+		{
+			Simulation->DestroyActor(SimulationActorHandle);
 		}
 	}
+
+	BodyRecords.Empty();
+	JointRecords.Empty();
 	ControlRecords.Empty();
 
-	if (CollisionActorHandle)
-	{
-		Simulation->DestroyActor(CollisionActorHandle);
-	}
 	CollisionActorHandle = nullptr;
-
-	if (SimulationActorHandle)
-	{
-		Simulation->DestroyActor(SimulationActorHandle);
-	}
 	SimulationActorHandle = nullptr;
 
 	Simulation.Reset();
@@ -619,7 +614,7 @@ ImmediatePhysics::FActorHandle* FRigPhysicsSimulation::CreateBody(
 void FRigPhysicsSimulation::Instantiate(
 	const FRigVMExecuteContext&       ExecuteContext, 
 	URigHierarchy&                    Hierarchy, 
-	const FRigPhysicsSolverComponent* SolverComponent)
+	const FRigPhysicsSolverComponent& SolverComponent)
 {
 	DECLARE_SCOPE_HIERARCHICAL_COUNTER_FUNC()
 	QUICK_SCOPE_CYCLE_COUNTER(STAT_RigPhysics_Instantiate);
@@ -638,7 +633,7 @@ void FRigPhysicsSimulation::Instantiate(
 	UpdateSimulationSpaceStateAndCalculateData(
 		ExecuteContext, Hierarchy, SolverComponent, 0.0f, ExecuteContext.GetAbsoluteTime());
 
-	InitialiseSimulation(SolverComponent);
+	InitialiseSimulation();
 
 	InitialiseBodyRecords(Hierarchy, SolverComponent);
 
@@ -676,19 +671,19 @@ ImmediatePhysics::FActorHandle* FRigPhysicsSimulation::GetActor(const FRigCompon
 
 //======================================================================================================================
 void FRigPhysicsSimulation::InstantiateSolverCollision(
-	const FRigPhysicsSolverComponent* SolverComponent,
+	const FRigPhysicsSolverComponent& SolverComponent,
 	FRigPhysicsIgnorePairs&           IgnorePairs)
 {
 	// Optionally create an object to contain environment collision
-	if (!SolverComponent->SolverSettings.Collision.IsEmpty())
+	if (!SolverComponent.SolverSettings.Collision.IsEmpty())
 	{
 		// When we make these additional collision shapes, their actors are all considered to be at
 		// the origin, with the offsets being contained in the collision shapes.
 		FTransform BodyRelSimSpaceTM = ConvertCollisionSpaceTransformToSimSpace(
-			SolverComponent->SolverSettings, FTransform());
+			SolverComponent.SolverSettings, FTransform());
 
 		CollisionActorHandle = CreateBody(
-			FName(TEXT("Environment")), SolverComponent->SolverSettings.Collision, nullptr, nullptr, BodyRelSimSpaceTM);
+			FName(TEXT("Environment")), SolverComponent.SolverSettings.Collision, nullptr, nullptr, BodyRelSimSpaceTM);
 	}
 
 	// Add no-collision pairs
@@ -711,14 +706,14 @@ void FRigPhysicsSimulation::InstantiateSolverCollision(
 //======================================================================================================================
 void FRigPhysicsSimulation::InstantiatePhysicsBodies(
 	URigHierarchy&                    Hierarchy,
-	const FRigPhysicsSolverComponent* SolverComponent,
+	const FRigPhysicsSolverComponent& SolverComponent,
 	FRigPhysicsIgnorePairs&           IgnorePairs)
 {
 	DECLARE_SCOPE_HIERARCHICAL_COUNTER_FUNC()
 	QUICK_SCOPE_CYCLE_COUNTER(STAT_RigPhysics_InstantiatePhysicsBodies);
 
-	const FRigPhysicsSolverSettings& SolverSettings = SolverComponent->SolverSettings;
-	const FRigPhysicsSimulationSpaceSettings& SimulationSpaceSettings = SolverComponent->SimulationSpaceSettings;
+	const FRigPhysicsSolverSettings& SolverSettings = SolverComponent.SolverSettings;
+	const FRigPhysicsSimulationSpaceSettings& SimulationSpaceSettings = SolverComponent.SimulationSpaceSettings;
 
 	for (TPair<FRigComponentKey, FRigBodyRecord>& BodyRecordPair : BodyRecords)
 	{
@@ -743,7 +738,8 @@ void FRigPhysicsSimulation::InstantiatePhysicsBodies(
 			{
 				Record.FinalComponentSpaceTM = Hierarchy.GetGlobalTransform(SourceKey);
 				const FTransform SourceSimulationSpaceTM =
-					ConvertComponentSpaceTransformToSimSpace(SolverComponent->SolverSettings, Record.FinalComponentSpaceTM);
+					ConvertComponentSpaceTransformToSimSpace(
+						SolverComponent.SolverSettings, Record.FinalComponentSpaceTM);
 				Record.ActorHandle = CreateBody(
 					ComponentKey.ElementKey.Name, Collision, &Dynamics, &BodyData, SourceSimulationSpaceTM);
 				if (Record.ActorHandle)
@@ -783,14 +779,14 @@ void FRigPhysicsSimulation::SetPhysicsBodyComponentState(
 //======================================================================================================================
 void FRigPhysicsSimulation::InstantiatePhysicsJoints(
 	const URigHierarchy&              Hierarchy,
-	const FRigPhysicsSolverComponent* SolverComponent,
+	const FRigPhysicsSolverComponent& SolverComponent,
 	FRigPhysicsIgnorePairs&           IgnorePairs)
 {
 	DECLARE_SCOPE_HIERARCHICAL_COUNTER_FUNC()
 	QUICK_SCOPE_CYCLE_COUNTER(STAT_RigPhysics_InstantiatePhysicsJoints);
 
-	const FRigPhysicsSolverSettings& SolverSettings = SolverComponent->SolverSettings;
-	const FRigPhysicsSimulationSpaceSettings& SimulationSpaceSettings = SolverComponent->SimulationSpaceSettings;
+	const FRigPhysicsSolverSettings& SolverSettings = SolverComponent.SolverSettings;
+	const FRigPhysicsSimulationSpaceSettings& SimulationSpaceSettings = SolverComponent.SimulationSpaceSettings;
 
 	// Once all the bodies are created, we can make their physics joints
 	for (TPair<FRigComponentKey, FRigJointRecord>& JointRecordPair : JointRecords)
@@ -837,17 +833,20 @@ void FRigPhysicsSimulation::InstantiatePhysicsJoints(
 				Chaos::FPBDJointSettings JointSettings;
 				if (JointData.bAutoCalculateChildOffset)
 				{
-					JointSettings.ConnectorTransforms[0] = FTransform();
+					JointSettings.ConnectorTransforms[ConstraintChildIndex] = FTransform();
 				}
-				JointSettings.ConnectorTransforms[0] =
-					JointData.ExtraChildOffset * JointSettings.ConnectorTransforms[0];
+				JointSettings.ConnectorTransforms[ConstraintChildIndex] =
+					JointData.ExtraChildOffset * JointSettings.ConnectorTransforms[ConstraintChildIndex];
 
 				if (JointData.bAutoCalculateParentOffset)
 				{
-					JointSettings.ConnectorTransforms[1] = Hierarchy.GetLocalTransform(ChildBoneKey, true);
+					// When auto-calculating, the transform stored in the parent frame should place
+					// it at the child bone's location.
+					JointSettings.ConnectorTransforms[ConstraintParentIndex] = 
+						Hierarchy.GetLocalTransform(ChildBoneKey, true);
 				}
-				JointSettings.ConnectorTransforms[1] =
-					JointData.ExtraParentOffset * JointSettings.ConnectorTransforms[1];
+				JointSettings.ConnectorTransforms[ConstraintParentIndex] =
+					JointData.ExtraParentOffset * JointSettings.ConnectorTransforms[ConstraintParentIndex];
 
 				ImmediatePhysics::UpdateJointSettingsFromLinearConstraint(JointData.LinearConstraint, JointSettings);
 				ImmediatePhysics::UpdateJointSettingsFromConeConstraint(JointData.ConeConstraint, JointSettings);
@@ -884,12 +883,17 @@ void FRigPhysicsSimulation::InstantiatePhysicsJoints(
 //======================================================================================================================
 void FRigPhysicsSimulation::InstantiateControls(
 	const URigHierarchy&              Hierarchy,
-	const FRigPhysicsSolverComponent* SolverComponent,
+	const FRigPhysicsSolverComponent& SolverComponent,
 	FRigPhysicsIgnorePairs&           IgnorePairs)
 {
 	DECLARE_SCOPE_HIERARCHICAL_COUNTER_FUNC()
 	QUICK_SCOPE_CYCLE_COUNTER(STAT_RigPhysics_InstantiateControls);
 
+	if (!Simulation.IsValid())
+	{
+		return;
+	}
+
 	for (TPair<FRigComponentKey, FRigControlRecord>& ControlRecordPair : ControlRecords)
 	{
 		const FRigComponentKey& ComponentKey = ControlRecordPair.Key;
@@ -921,8 +925,11 @@ void FRigPhysicsSimulation::InstantiateControls(
 				ImmediatePhysics::FActorHandle* const ParentBodyHandle = 
 					ParentBodyRecord ? ParentBodyRecord->ActorHandle : SimulationActorHandle;
 
-				// This handles nullptrs. The constraint is created disabled - it will be updated in pre-physics
-				ControlRecord.JointHandle = CreatePhysicsJoint(Simulation.Get(), ChildBodyHandle, ParentBodyHandle);
+				if (ChildBodyHandle && ParentBodyHandle)
+				{
+					// The constraint is created disabled - it will be updated in pre-physics
+					ControlRecord.JointHandle = CreatePhysicsJoint(*Simulation.Get(), *ChildBodyHandle, *ParentBodyHandle);
+				}
 
 				if (!ControlRecord.JointHandle)
 				{
@@ -944,7 +951,7 @@ void FRigPhysicsSimulation::InstantiateControls(
 //======================================================================================================================
 void FRigPhysicsSimulation::UpdateBodyRecordPrePhysics(
 	const URigHierarchy&              Hierarchy,
-	const FRigPhysicsSolverComponent* SolverComponent,
+	const FRigPhysicsSolverComponent& SolverComponent,
 	const float                       DeltaTime,
 	FRigBodyRecord&                   Record, 
 	const FRigPhysicsBodyComponent*   PhysicsBodyComponent)
@@ -988,16 +995,20 @@ void FRigPhysicsSimulation::UpdateBodyRecordPrePhysics(
 //======================================================================================================================
 void FRigPhysicsSimulation::UpdateBodyPrePhysics(
 	const FRigVMExecuteContext&       ExecuteContext, 
-	const FRigPhysicsSolverComponent* SolverComponent,
+	const FRigPhysicsSolverComponent& SolverComponent,
 	const FRigBodyRecord&             Record, 
-	const FRigPhysicsBodyComponent*   PhysicsBodyComponent)
+	const FRigPhysicsBodyComponent&   PhysicsBodyComponent)
 {
-	const FRigComponentKey ComponentKey = PhysicsBodyComponent->GetKey();
+	if (!Record.ActorHandle)
+	{
+		return;
+	}
+	const FRigComponentKey ComponentKey = PhysicsBodyComponent.GetKey();
 
-	FPhysicsControlModifierData BodyData = PhysicsBodyComponent->BodyData;
+	FPhysicsControlModifierData BodyData = PhysicsBodyComponent.BodyData;
 	EPhysicsControlKinematicTargetSpace KinematicTargetSpace = BodyData.KinematicTargetSpace;
 
-	if (SolverComponent->TrackInputCounter > 0)
+	if (SolverComponent.TrackInputCounter > 0)
 	{
 		// If the body is already kinematic, then we don't want to make it start ignoring a target.
 		// However, dynamic objects should avoid using a target if they are just tracking due to a reset.
@@ -1009,7 +1020,7 @@ void FRigPhysicsSimulation::UpdateBodyPrePhysics(
 	}
 
 	UpdateBodyFromModifierData(
-		Record.ActorHandle, Simulation.Get(), BodyData, SimulationSpaceData.Gravity);
+		*Record.ActorHandle, *Simulation.Get(), BodyData, SimulationSpaceData.Gravity);
 
 	if (Record.ActorHandle->GetIsKinematic())
 	{
@@ -1017,25 +1028,24 @@ void FRigPhysicsSimulation::UpdateBodyPrePhysics(
 
 		// If the target is already in component space, then that's all we need
 		FTransform KinematicTargetCS;
+		const FTransform KinematicTargetTM = PhysicsBodyComponent.KinematicTarget;
 		switch (KinematicTargetSpace)
 		{
 		case EPhysicsControlKinematicTargetSpace::Component:
 		{
-			KinematicTargetCS = PhysicsBodyComponent->KinematicTarget;
+			KinematicTargetCS = KinematicTargetTM;
 			break;
 		}
 		case EPhysicsControlKinematicTargetSpace::World:
 		{
-			// Record.KinematicTarget * SimulationSpaceState.ComponentTM.Inverse();
-			KinematicTargetCS = PhysicsBodyComponent->KinematicTarget.GetRelativeTransform(
-				SimulationSpaceState.ComponentTM);
+			KinematicTargetCS = KinematicTargetTM.GetRelativeTransform(SimulationSpaceState.ComponentTM);
 			break;
 		}
 		default:
 		{
 			// All the other options are relative to a bone, so the first task is to get
 			// that, which will be in component space
-			FRigElementKey SourceKey = PhysicsBodyComponent->BodySolverSettings.SourceBone;
+			FRigElementKey SourceKey = PhysicsBodyComponent.BodySolverSettings.SourceBone;
 			if (!SourceKey.IsValid())
 			{
 				SourceKey = ComponentKey.ElementKey;
@@ -1046,21 +1056,19 @@ void FRigPhysicsSimulation::UpdateBodyPrePhysics(
 				{
 				case EPhysicsControlKinematicTargetSpace::OffsetInBoneSpace:
 				{
-					KinematicTargetCS = PhysicsBodyComponent->KinematicTarget * Record.SourceComponentSpaceTM.ToTransform();
+					KinematicTargetCS = KinematicTargetTM * Record.SourceComponentSpaceTM;
 					break;
 				}
 				case EPhysicsControlKinematicTargetSpace::OffsetInWorldSpace:
 				{
+					const FTransform ComponentTM(SimulationSpaceState.ComponentTM);
 					// Convert the bone to WS (it's really component space!), apply the target, and convert back. 
-					FTransform BoneWS = Record.SourceComponentSpaceTM.ToTransform() * SimulationSpaceState.ComponentTM;
-					FTransform KinematicTargetWS = BoneWS;
+					FTransform KinematicTargetWS = Record.SourceComponentSpaceTM * ComponentTM;
 					// Note that we don't want to rotate the translation by the target - we apply it
 					// individually to orientation and position.
-					KinematicTargetWS.AddToTranslation(PhysicsBodyComponent->KinematicTarget.GetTranslation());
-					KinematicTargetWS.SetRotation(
-						PhysicsBodyComponent->KinematicTarget.GetRotation() * KinematicTargetWS.GetRotation());
-					// Danny TODO figure out which of the GetRelativeTransform versions this is
-					KinematicTargetCS = SimulationSpaceState.ComponentTM.Inverse() * KinematicTargetWS;
+					KinematicTargetWS.AddToTranslation(KinematicTargetTM.GetTranslation());
+					KinematicTargetWS.SetRotation(KinematicTargetTM.GetRotation() * KinematicTargetWS.GetRotation());
+					KinematicTargetCS = KinematicTargetWS.GetRelativeTransform(ComponentTM);
 					break;
 				}
 				case EPhysicsControlKinematicTargetSpace::OffsetInComponentSpace:
@@ -1068,51 +1076,51 @@ void FRigPhysicsSimulation::UpdateBodyPrePhysics(
 					// This applies the offset as a transform in component space. Note that this is
 					// different to what PCC does, because in ControlRig "world" is really
 					// "component" space.
-					KinematicTargetCS = Record.SourceComponentSpaceTM.ToTransform() * PhysicsBodyComponent->KinematicTarget;
+					KinematicTargetCS =  Record.SourceComponentSpaceTM * KinematicTargetTM;
 					break;
 				}
 				case EPhysicsControlKinematicTargetSpace::IgnoreTarget:
 				{
-					KinematicTargetCS = Record.SourceComponentSpaceTM.ToTransform();
+					KinematicTargetCS = Record.SourceComponentSpaceTM;
 					break;
 				}
 				default:
 				{
-					KinematicTargetCS = Record.SourceComponentSpaceTM.ToTransform();
+					KinematicTargetCS = Record.SourceComponentSpaceTM;
 					UE_CONTROLRIG_RIGUNIT_REPORT_ERROR(TEXT("Kinematic target space is not valid"));
 				}
 				}
 			}
 		}
 		}
-		FTransform KinematicTargetTM = ConvertComponentSpaceTransformToSimSpace(
-			SolverComponent->SolverSettings, KinematicTargetCS);
-		Record.ActorHandle->SetKinematicTarget(KinematicTargetTM);
+		const FTransform KinematicTargetSimSpaceTM = ConvertComponentSpaceTransformToSimSpace(
+			SolverComponent.SolverSettings, KinematicTargetCS);
+		Record.ActorHandle->SetKinematicTarget(KinematicTargetSimSpaceTM);
 	}
 	else
 	{
 		// Danny TODO move damping into BodyData - any PhysicsControl system should be able to use it
-		Record.ActorHandle->SetLinearDamping(PhysicsBodyComponent->Dynamics.LinearDamping);
-		Record.ActorHandle->SetAngularDamping(PhysicsBodyComponent->Dynamics.AngularDamping);
+		Record.ActorHandle->SetLinearDamping(PhysicsBodyComponent.Dynamics.LinearDamping);
+		Record.ActorHandle->SetAngularDamping(PhysicsBodyComponent.Dynamics.AngularDamping);
 	}
 }
 
 //======================================================================================================================
 void FRigPhysicsSimulation::ApplyForcesPrePhysics(
 	const FRigVMExecuteContext&       ExecuteContext,
-	const FRigPhysicsSolverComponent* SolverComponent,
+	const FRigPhysicsSolverComponent& SolverComponent,
 	const FRigBodyRecord&             Record,
-	const FRigPhysicsBodyComponent*   PhysicsBodyComponent)
+	const FRigPhysicsBodyComponent&   PhysicsBodyComponent)
 {
-	if (PhysicsBodyComponent->ForceAndTorques.IsEmpty() || !Record.ActorHandle || Record.ActorHandle->GetIsKinematic())
+	if (PhysicsBodyComponent.ForceAndTorques.IsEmpty() || !Record.ActorHandle || Record.ActorHandle->GetIsKinematic())
 	{
 		return;
 	}
 
-	ImmediatePhysics::FActorHandle* ActorHandle = Record.ActorHandle;
-	const FRigPhysicsSolverSettings& SolverSettings = SolverComponent->SolverSettings;
+	ImmediatePhysics::FActorHandle& ActorHandle = *Record.ActorHandle;
+	const FRigPhysicsSolverSettings& SolverSettings = SolverComponent.SolverSettings;
 
-	for (const FPhysicsControlNamedForceAndTorqueData& Data : PhysicsBodyComponent->ForceAndTorques)
+	for (const FPhysicsControlNamedForceAndTorqueData& Data : PhysicsBodyComponent.ForceAndTorques)
 	{
 		const FPhysicsControlForceAndTorqueData& ForceAndTorqueData = Data.ForceAndTorqueData;
 
@@ -1120,7 +1128,7 @@ void FRigPhysicsSimulation::ApplyForcesPrePhysics(
 		FVector LinearSS;
 		FVector AngularSS;
 
-		FTransform BodyTM = ActorHandle->GetWorldTransform();
+		FTransform BodyTM = ActorHandle.GetWorldTransform();
 
 		switch (ForceAndTorqueData.ForceAndTorqueSpace)
 		{
@@ -1150,21 +1158,21 @@ void FRigPhysicsSimulation::ApplyForcesPrePhysics(
 		{
 			FVector SimSpaceLocation = ConvertWorldPositionToSimSpaceNoScale(
 				SolverSettings, ForceAndTorqueData.Location);
-			ActorHandle->AddForceAtLocation(LinearSS, SimSpaceLocation, FT);
+			ActorHandle.AddForceAtLocation(LinearSS, SimSpaceLocation, FT);
 		}
 		break;
 		case EPhysicsControlSpace::Body:
 		{
-			FVector SimSpaceLocation = ActorHandle->GetWorldTransform().TransformPositionNoScale(
-				ActorHandle->GetLocalCoMLocation() + ForceAndTorqueData.Location);
-			ActorHandle->AddForceAtLocation(LinearSS, SimSpaceLocation, FT);
+			FVector SimSpaceLocation = ActorHandle.GetWorldTransform().TransformPositionNoScale(
+				ActorHandle.GetLocalCoMLocation() + ForceAndTorqueData.Location);
+			ActorHandle.AddForceAtLocation(LinearSS, SimSpaceLocation, FT);
 		}
 		break;
 		default:
 			break;
 		}
 		// Torque doesn't get applied at a location
-		ActorHandle->AddTorque(AngularSS, FT);
+		ActorHandle.AddTorque(AngularSS, FT);
 	}
 }
 
@@ -1172,7 +1180,7 @@ void FRigPhysicsSimulation::ApplyForcesPrePhysics(
 void FRigPhysicsSimulation::UpdateJointPrePhysics(
 	const URigHierarchy&                Hierarchy, 
 	FRigJointRecord&                    Record, 
-	const FRigPhysicsJointComponent*    PhysicsJointComponent, 
+	const FRigPhysicsJointComponent&    PhysicsJointComponent, 
 	const float                         DeltaTime)
 {
 	// Now update the joint targets
@@ -1182,30 +1190,46 @@ void FRigPhysicsSimulation::UpdateJointPrePhysics(
 	}
 	if (Chaos::FPBDJointConstraintHandle* Constraint = Record.JointHandle->GetConstraint())
 	{
-		const FRigComponentKey ComponentKey = PhysicsJointComponent->GetKey();
+		const FRigComponentKey ComponentKey = PhysicsJointComponent.GetKey();
 
 		// Set the drive strength etc
-		const FRigPhysicsJointData& JointData = PhysicsJointComponent->JointData;
-		const FRigPhysicsDriveData& DriveData = PhysicsJointComponent->DriveData;
+		const FRigPhysicsJointData& JointData = PhysicsJointComponent.JointData;
+		const FRigPhysicsDriveData& DriveData = PhysicsJointComponent.DriveData;
 
-		Constraint->SetConstraintEnabled(JointData.bEnabled);
+		if (PhysicsJointComponent.DirtyFlags != ERigPhysicsJointComponentDirtyFlags::None)
+		{
+			Chaos::FPBDJointSettings Settings = Constraint->GetSettings();
+
+			if ((PhysicsJointComponent.DirtyFlags & ERigPhysicsJointComponentDirtyFlags::Joint) != 
+				ERigPhysicsJointComponentDirtyFlags::None)
+			{
+				SetPhysicsJointEnabled(*Record.JointHandle, JointData.bEnabled);
 
-		Chaos::FPBDJointSettings Settings = Constraint->GetSettings();
+				ImmediatePhysics::UpdateJointSettingsFromLinearConstraint(JointData.LinearConstraint, Settings);
+				ImmediatePhysics::UpdateJointSettingsFromConeConstraint(JointData.ConeConstraint, Settings);
+				ImmediatePhysics::UpdateJointSettingsFromTwistConstraint(JointData.TwistConstraint, Settings);
+			}
+
+			if ((PhysicsJointComponent.DirtyFlags & ERigPhysicsJointComponentDirtyFlags::Drive) !=
+				ERigPhysicsJointComponentDirtyFlags::None)
+			{
+				ImmediatePhysics::UpdateJointSettingsFromLinearDriveConstraint(DriveData.LinearDriveConstraint, Settings);
+				ImmediatePhysics::UpdateJointSettingsFromAngularDriveConstraint(DriveData.AngularDriveConstraint, Settings);
+			}
 
-		ImmediatePhysics::UpdateJointSettingsFromLinearConstraint(JointData.LinearConstraint, Settings);
-		ImmediatePhysics::UpdateJointSettingsFromConeConstraint(JointData.ConeConstraint, Settings);
-		ImmediatePhysics::UpdateJointSettingsFromTwistConstraint(JointData.TwistConstraint, Settings);
+			Constraint->SetSettings(Settings);
 
-		ImmediatePhysics::UpdateJointSettingsFromLinearDriveConstraint(DriveData.LinearDriveConstraint, Settings);
-		ImmediatePhysics::UpdateJointSettingsFromAngularDriveConstraint(DriveData.AngularDriveConstraint, Settings);
+			PhysicsJointComponent.DirtyFlags = ERigPhysicsJointComponentDirtyFlags::None;
+		}
 
-		Constraint->SetSettings(Settings);
+		const Chaos::FPBDJointSettings& Settings = Constraint->GetSettings();
 
 		// Now set the actual target
-		if (Settings.AngularDriveStiffness.SquaredLength() > 0.0f ||
+		if (JointData.bEnabled && (
+			Settings.AngularDriveStiffness.SquaredLength() > 0.0f ||
 			Settings.AngularDriveDamping.SquaredLength() > 0.0f ||
 			Settings.LinearDriveStiffness.SquaredLength() > 0.0f ||
-			Settings.LinearDriveDamping.SquaredLength() > 0.0f)
+			Settings.LinearDriveDamping.SquaredLength() > 0.0f))
 		{
 			// Multiplier on the velocity calculated from the current and previous target
 			//
@@ -1225,7 +1249,7 @@ void FRigPhysicsSimulation::UpdateJointPrePhysics(
 				TargetVelocityMultiplier = DriveData.SkeletalAnimationVelocityMultiplier;
 			}
 
-			UE::PhysicsControl::FPosQuat DriveTargetTM;
+			FTransform DriveTargetTM;
 			if (DriveData.bUseSkeletalAnimation)
 			{
 				FRigElementKey ParentSourceKey;
@@ -1249,32 +1273,31 @@ void FRigPhysicsSimulation::UpdateJointPrePhysics(
 
 				// Note that the drive operates between a parent and child part, so we
 				// don't need to worry about global/component (etc) space.
-				const FTransform ChildTM = Hierarchy.GetGlobalTransform(ChildSourceKey);
-				const FTransform ParentTM = Hierarchy.GetGlobalTransform(ParentSourceKey);
+				const FTransform ChildTM(Hierarchy.GetGlobalTransform(ChildSourceKey));
+				const FTransform ParentTM(Hierarchy.GetGlobalTransform(ParentSourceKey));
 
-				const UE::PhysicsControl::FPosQuat ComponentSpaceParentFrameTM = 
-					Settings.ConnectorTransforms[1] * ParentTM;
-				const UE::PhysicsControl::FPosQuat ComponentSpaceChildFrameTM = 
-					Settings.ConnectorTransforms[0] * ChildTM;
+				const FTransform ComponentSpaceParentFrameTM = 
+					Settings.ConnectorTransforms[ConstraintParentIndex] * ParentTM;
+				const FTransform ComponentSpaceChildFrameTM = 
+					Settings.ConnectorTransforms[ConstraintChildIndex] * ChildTM;
 
-				DriveTargetTM = ComponentSpaceParentFrameTM.Inverse() * ComponentSpaceChildFrameTM;
+				DriveTargetTM = ComponentSpaceChildFrameTM.GetRelativeTransform(ComponentSpaceParentFrameTM);
 			}
 
 			// Apply the offset in the frame of the (potentially animation) target
-			DriveTargetTM = DriveTargetTM * UE::PhysicsControl::FPosQuat(
-				DriveData.LinearDriveConstraint.PositionTarget, 
-				DriveData.AngularDriveConstraint.OrientationTarget.Quaternion());
+			DriveTargetTM = FTransform(				
+				DriveData.AngularDriveConstraint.OrientationTarget.Quaternion(),
+				DriveData.LinearDriveConstraint.PositionTarget) * DriveTargetTM;
 
 			Constraint->SetLinearDrivePositionTarget(DriveTargetTM.GetTranslation());
 			Constraint->SetAngularDrivePositionTarget(DriveTargetTM.GetRotation());
 
-			if (Record.PreviousDriveTargetUpdateCounter + 1 == UpdateCounter &&
-				TargetVelocityMultiplier > 0.0f)
+			if (Record.PreviousDriveTargetUpdateCounter + 1 == UpdateCounter && TargetVelocityMultiplier > 0.0f)
 			{
 				if (DeltaTime > SMALL_NUMBER)
 				{
-					const UE::PhysicsControl::FPosQuat DriveTargetTMDelta =
-						DriveTargetTM * Record.PreviousDriveTargetTM.Inverse();
+					const FTransform DriveTargetTMDelta =
+						Record.PreviousDriveTargetTM.GetRelativeTransformReverse(DriveTargetTM);
 					const FVector Velocity = DriveTargetTMDelta.GetTranslation() / DeltaTime;
 					const FVector AngularVelocity =
 						DriveTargetTMDelta.GetRotation().GetShortestArcWith(
@@ -1298,38 +1321,34 @@ void FRigPhysicsSimulation::UpdateJointPrePhysics(
 }
 
 //======================================================================================================================
-static UE::PhysicsControl::FPosQuat CalculateTargetTM(
+static FTransform CalculateTargetTM(
 	const URigHierarchy&               Hierarchy,
 	const Chaos::FPBDJointSettings&    JointSettings,
 	const FRigControlRecord&           Record)
 {
 	const FTransform ChildTM = Record.ChildBodyComponentKey.IsValid() ?
-		Hierarchy.GetGlobalTransform(Record.ChildBodyComponentKey.ElementKey) :
-		FTransform();
+		Hierarchy.GetGlobalTransform(Record.ChildBodyComponentKey.ElementKey) : FTransform();
 
-	const UE::PhysicsControl::FPosQuat ChildTargetTM =
-		UE::PhysicsControl::FPosQuat(ChildTM) * 
-		UE::PhysicsControl::FPosQuat(JointSettings.ConnectorTransforms[ConstraintChildIndex]);
+	const FTransform ChildTargetTM = JointSettings.ConnectorTransforms[ConstraintChildIndex] * ChildTM;
 
 	if (Record.ParentBodyComponentKey.IsValid())
 	{
-		const FTransform ParentTM = Hierarchy.GetGlobalTransform(Record.ParentBodyComponentKey.ElementKey);
-		const UE::PhysicsControl::FPosQuat ParentTargetTM =
-			UE::PhysicsControl::FPosQuat(ParentTM) *
-			UE::PhysicsControl::FPosQuat(JointSettings.ConnectorTransforms[ConstraintParentIndex]);
-		return ParentTargetTM.Inverse() * ChildTargetTM;
+		const FTransform ParentTM(
+			Hierarchy.GetGlobalTransform(Record.ParentBodyComponentKey.ElementKey));
+		const FTransform ParentTargetTM = JointSettings.ConnectorTransforms[ConstraintParentIndex] * ParentTM;
+		return ChildTargetTM.GetRelativeTransform(ParentTargetTM);
 	}
 	return ChildTargetTM;
 }
 
 //======================================================================================================================
 void FRigPhysicsSimulation::CheckForResetsPrePhysics(
-	const URigHierarchy& Hierarchy, FRigPhysicsSolverComponent* SolverComponent, const float DeltaTime)
+	const URigHierarchy& Hierarchy, FRigPhysicsSolverComponent& SolverComponent, const float DeltaTime)
 {
 	DECLARE_SCOPE_HIERARCHICAL_COUNTER_FUNC()
 	QUICK_SCOPE_CYCLE_COUNTER(STAT_RigPhysics_CheckForResetsPrePhysics);
 
-	const FRigPhysicsSolverSettings& SolverSettings = SolverComponent->SolverSettings;
+	const FRigPhysicsSolverSettings& SolverSettings = SolverComponent.SolverSettings;
 
 	double SpeedThresholdForResetSquared = SolverSettings.KinematicSpeedThresholdForReset > 0 ?
 		FMath::Square(SolverSettings.KinematicSpeedThresholdForReset) : TNumericLimits<double>::Max();
@@ -1382,7 +1401,7 @@ void FRigPhysicsSimulation::CheckForResetsPrePhysics(
 			UE_LOG(LogRigPhysics, Log, TEXT("Acceleration %f triggered reset in %s"),
 				FMath::Sqrt(HighestAccelerationSq), *OwnerName.ToString());
 		}
-		SolverComponent->TrackInputCounter = FMath::Max(SolverComponent->TrackInputCounter, 3);
+		SolverComponent.TrackInputCounter = FMath::Max(SolverComponent.TrackInputCounter, 3);
 	}
 }
 
@@ -1390,35 +1409,45 @@ void FRigPhysicsSimulation::CheckForResetsPrePhysics(
 void FRigPhysicsSimulation::UpdatePrePhysics(
 	const FRigVMExecuteContext& ExecuteContext, 
 	const URigHierarchy&        Hierarchy,
-	FRigPhysicsSolverComponent* SolverComponent, 
+	FRigPhysicsSolverComponent& SolverComponent, 
 	const float                 DeltaTime)
 {
 	DECLARE_SCOPE_HIERARCHICAL_COUNTER_FUNC()
 	QUICK_SCOPE_CYCLE_COUNTER(STAT_RigPhysics_UpdatePrePhysics);
 
+	if (!Simulation.IsValid())
+	{
+		return;
+	}
+
 	if (CollisionActorHandle)
 	{
 		FTransform BodyRelSimSpaceTM = ConvertCollisionSpaceTransformToSimSpace(
-			SolverComponent->SolverSettings, FTransform());
+			SolverComponent.SolverSettings, FTransform());
 		CollisionActorHandle->SetKinematicTarget(BodyRelSimSpaceTM);
 	}
 
-	for (TPair<FRigComponentKey, FRigBodyRecord>& BodyRecordPair : BodyRecords)
 	{
-		const FRigComponentKey& ComponentKey = BodyRecordPair.Key;
-		FRigBodyRecord& Record = BodyRecordPair.Value;
-		if (const FRigPhysicsBodyComponent* PhysicsBodyComponent = Cast<FRigPhysicsBodyComponent>(
-			Hierarchy.FindComponent(ComponentKey)))
+		DECLARE_SCOPE_HIERARCHICAL_COUNTER_FUNC()
+		QUICK_SCOPE_CYCLE_COUNTER(STAT_RigPhysics_UpdateBodyRecordsPrePhysics);
+		for (TPair<FRigComponentKey, FRigBodyRecord>& BodyRecordPair : BodyRecords)
 		{
-			UpdateBodyRecordPrePhysics(Hierarchy, SolverComponent, DeltaTime, Record, PhysicsBodyComponent);
+			const FRigComponentKey& ComponentKey = BodyRecordPair.Key;
+			FRigBodyRecord& Record = BodyRecordPair.Value;
+			if (const FRigPhysicsBodyComponent* PhysicsBodyComponent = Cast<FRigPhysicsBodyComponent>(
+				Hierarchy.FindComponent(ComponentKey)))
+			{
+				UpdateBodyRecordPrePhysics(Hierarchy, SolverComponent, DeltaTime, Record, PhysicsBodyComponent);
+			}
 		}
 	}
 
 	CheckForResetsPrePhysics(Hierarchy, SolverComponent, DeltaTime);
-	if (SolverComponent->TrackInputCounter > 0)
+
+	if (SolverComponent.TrackInputCounter > 0)
 	{
 		UE_LOG(LogRigPhysics, Log, TEXT("Forcing tracking (counter = %d) of input for %s"), 
-			SolverComponent->TrackInputCounter, *OwnerName.ToString());
+			SolverComponent.TrackInputCounter, *OwnerName.ToString());
 	}
 
 	{
@@ -1431,7 +1460,7 @@ void FRigPhysicsSimulation::UpdatePrePhysics(
 			if (const FRigPhysicsBodyComponent* PhysicsBodyComponent = Cast<FRigPhysicsBodyComponent>(
 				Hierarchy.FindComponent(ComponentKey)))
 			{
-				UpdateBodyPrePhysics(ExecuteContext, SolverComponent, Record, PhysicsBodyComponent);
+				UpdateBodyPrePhysics(ExecuteContext, SolverComponent, Record, *PhysicsBodyComponent);
 			}
 		}
 	}
@@ -1447,7 +1476,7 @@ void FRigPhysicsSimulation::UpdatePrePhysics(
 			if (const FRigPhysicsJointComponent* JointComponent = Cast<FRigPhysicsJointComponent>(
 				Hierarchy.FindComponent(ComponentKey)))
 			{
-				UpdateJointPrePhysics(Hierarchy, JointRecord, JointComponent, DeltaTime);
+				UpdateJointPrePhysics(Hierarchy, JointRecord, *JointComponent, DeltaTime);
 			}
 		}
 	}
@@ -1463,11 +1492,7 @@ void FRigPhysicsSimulation::UpdatePrePhysics(
 			if (const FRigPhysicsControlComponent* ControlComponent = Cast<FRigPhysicsControlComponent>(
 				Hierarchy.FindComponent(ComponentKey)))
 			{
-				if (ControlComponent->ControlData.bEnabled)
-				{
-					UpdateControlPrePhysics(ControlRecord, ControlComponent, SolverComponent, Hierarchy, DeltaTime);
-				}
-				SetPhysicsJointEnabled(ControlRecord.JointHandle, ControlComponent->ControlData.bEnabled);
+				UpdateControlPrePhysics(ControlRecord, *ControlComponent, SolverComponent, Hierarchy, DeltaTime);
 			}
 		}
 	}
@@ -1479,10 +1504,11 @@ void FRigPhysicsSimulation::UpdatePrePhysics(
 		{
 			const FRigComponentKey& ComponentKey = BodyRecordPair.Key;
 			FRigBodyRecord& Record = BodyRecordPair.Value;
+
 			if (const FRigPhysicsBodyComponent* PhysicsBodyComponent = Cast<FRigPhysicsBodyComponent>(
 				Hierarchy.FindComponent(ComponentKey)))
 			{
-				ApplyForcesPrePhysics(ExecuteContext, SolverComponent, Record, PhysicsBodyComponent);
+				ApplyForcesPrePhysics(ExecuteContext, SolverComponent, Record, *PhysicsBodyComponent);
 			}
 
 		}
@@ -1492,12 +1518,12 @@ void FRigPhysicsSimulation::UpdatePrePhysics(
 //======================================================================================================================
 void FRigPhysicsSimulation::UpdateControlPrePhysics(
 	FRigControlRecord&                 ControlRecord,
-	const FRigPhysicsControlComponent* ControlComponent,
-	const FRigPhysicsSolverComponent*  SolverComponent,
+	const FRigPhysicsControlComponent& ControlComponent,
+	const FRigPhysicsSolverComponent&  SolverComponent,
 	const URigHierarchy&               Hierarchy,
 	const float                        DeltaTime)
 {
-	const FRigPhysicsSolverSettings& SolverSettings = SolverComponent->SolverSettings;
+	const FRigPhysicsSolverSettings& SolverSettings = SolverComponent.SolverSettings;
 
 	ImmediatePhysics::FJointHandle* JointHandle = ControlRecord.JointHandle;
 	if (!JointHandle)
@@ -1509,6 +1535,13 @@ void FRigPhysicsSimulation::UpdateControlPrePhysics(
 	{
 		return;
 	}
+
+	if (!ControlComponent.ControlData.bEnabled)
+	{
+		SetPhysicsJointEnabled(*JointHandle, false);
+		return;
+	}
+
 	float ThisDeltaTime = DeltaTime;
 
 	if (ControlRecord.PreviousTargetUpdateCounter + 1 != UpdateCounter)
@@ -1519,17 +1552,18 @@ void FRigPhysicsSimulation::UpdateControlPrePhysics(
 		ThisDeltaTime = 0.0f;
 	}
 
-	Constraint->SetCollisionEnabled(!ControlComponent->ControlData.bDisableCollision);
-	Constraint->SetParentInvMassScale(ControlComponent->ControlData.bOnlyControlChildObject ? 0 : 1);
+	if ((ControlComponent.DirtyFlags & ERigPhysicsControlComponentDirtyFlags::Data) !=
+		ERigPhysicsControlComponentDirtyFlags::None)
+	{
+		Constraint->SetCollisionEnabled(!ControlComponent.ControlData.bDisableCollision);
+		Constraint->SetParentInvMassScale(ControlComponent.ControlData.bOnlyControlChildObject ? 0 : 1);
+	}
 
 	// Check for early outs. Make sure that we don't apply velocities using the wrong calculation
 	// when the strength/damping is increased in the future by not updating the update counter.
 
-	const ImmediatePhysics::FActorHandle* ChildActorHandle =
-		JointHandle->GetActorHandles()[ConstraintChildIndex];
-	const ImmediatePhysics::FActorHandle* ParentActorHandle =
-		JointHandle->GetActorHandles()[ConstraintParentIndex];
-
+	const ImmediatePhysics::FActorHandle* ChildActorHandle = JointHandle->GetActorHandles()[ConstraintChildIndex];
+	const ImmediatePhysics::FActorHandle* ParentActorHandle = JointHandle->GetActorHandles()[ConstraintParentIndex];
 	if (!ChildActorHandle || !ParentActorHandle)
 	{
 		return;
@@ -1538,15 +1572,28 @@ void FRigPhysicsSimulation::UpdateControlPrePhysics(
 	const Chaos::FPBDJointSettings& JointSettings = Constraint->GetSettings();
 
 	// Note that if we don't have any strength, then we don't calculate the targets. 
-	if (!UpdateDriveSpringDamperSettings(
-		JointHandle, JointSettings, ControlComponent->ControlData, ControlComponent->ControlMultiplier))
+	if ((ControlComponent.DirtyFlags & ERigPhysicsControlComponentDirtyFlags::Data) !=
+		ERigPhysicsControlComponentDirtyFlags::None)
 	{
-		return;
+		ControlComponent.DirtyFlags &= ~ERigPhysicsControlComponentDirtyFlags::Data;
+		if (!UpdateDriveSpringDamperSettings(
+			*JointHandle, JointSettings, ControlComponent.ControlData, ControlComponent.ControlMultiplier))
+		{
+			ControlRecord.bCachedHasStrengthOrDamping = false;
+			return;
+		}
+		else
+		{
+			ControlRecord.bCachedHasStrengthOrDamping = true;
+		}
+	}
+	else
+	{
+		if (!ControlRecord.bCachedHasStrengthOrDamping)
+		{
+			return;
+		}
 	}
-
-	// TODO
-	// - cache settings / previous input parameters to avoid unnecessary repeating
-	//   calculations and making physics API calls every update.
 
 	// Update the target point on the child
 
@@ -1556,18 +1603,18 @@ void FRigPhysicsSimulation::UpdateControlPrePhysics(
 	// if the child is kinematic and the parent is dynamic. 
 	if (!ChildActorHandle->GetIsKinematic())
 	{
+		// Note - This simply sets the location and doesn't trigger additional work
 		Constraint->SetChildConnectorLocation(
-			ControlComponent->ControlData.GetControlPoint(ChildActorHandle));
+			ControlComponent.ControlData.GetControlPoint(ChildActorHandle));
 	}
 
 	FTransform TargetTM(
-		ControlComponent->ControlTarget.TargetOrientation,
-		ControlComponent->ControlTarget.TargetPosition);
+		ControlComponent.ControlTarget.TargetOrientation,
+		ControlComponent.ControlTarget.TargetPosition);
 
-	if (ControlComponent->ControlData.bUseSkeletalAnimation)
+	if (ControlComponent.ControlData.bUseSkeletalAnimation)
 	{
-		const FTransform ComponentSpaceAnimTargetTM = CalculateTargetTM(
-			Hierarchy, JointSettings, ControlRecord).ToTransform();
+		const FTransform ComponentSpaceAnimTargetTM = CalculateTargetTM(Hierarchy, JointSettings, ControlRecord);
 		const FTransform SimSpaceAnimTargetTM = ConvertComponentSpaceTransformToSimSpace(
 			SolverSettings, ComponentSpaceAnimTargetTM);
 		TargetTM = TargetTM * SimSpaceAnimTargetTM;
@@ -1580,35 +1627,34 @@ void FRigPhysicsSimulation::UpdateControlPrePhysics(
 	Constraint->SetLinearDrivePositionTarget(TargetTM.GetTranslation());
 	Constraint->SetAngularDrivePositionTarget(TargetTM.GetRotation());
 
-	if ((ThisDeltaTime * ControlComponent->ControlData.LinearTargetVelocityMultiplier) != 0)
+	if ((ThisDeltaTime * ControlComponent.ControlData.LinearTargetVelocityMultiplier) != 0)
 	{
-		FVector Velocity =
-			(TargetTM.GetTranslation() - ControlRecord.PreviousTargetTM.GetTranslation()) / ThisDeltaTime;
+		const FVector Velocity = UE::PhysicsControl::CalculateLinearVelocity(
+			ControlRecord.PreviousTargetTM.GetTranslation(), TargetTM.GetTranslation(), ThisDeltaTime);
 		Constraint->SetLinearDriveVelocityTarget(
-			Velocity * ControlComponent->ControlData.LinearTargetVelocityMultiplier);
+			Velocity * ControlComponent.ControlData.LinearTargetVelocityMultiplier);
 	}
 	else
 	{
 		Constraint->SetLinearDriveVelocityTarget(Chaos::FVec3(0));
 	}
 
-	if ((ThisDeltaTime * ControlComponent->ControlData.AngularTargetVelocityMultiplier) != 0)
+	if ((ThisDeltaTime * ControlComponent.ControlData.AngularTargetVelocityMultiplier) != 0)
 	{
-		// Note that quats multiply in the opposite order to TMs, and must be in the same hemisphere.
 		const FQuat Q = TargetTM.GetRotation();
-		FQuat PrevQ = ControlRecord.PreviousTargetTM.GetRotation();
-		PrevQ.EnforceShortestArcWith(Q);
-		const FQuat DeltaQ = Q * PrevQ.Inverse();
-		const FVector AngularVelocity = DeltaQ.ToRotationVector() / ThisDeltaTime;
-
+		const FQuat PrevQ = ControlRecord.PreviousTargetTM.GetRotation();
+		const FVector AngularVelocity = UE::PhysicsControl::CalculateAngularVelocity(PrevQ, Q, ThisDeltaTime);
 		Constraint->SetAngularDriveVelocityTarget(
-			AngularVelocity * ControlComponent->ControlData.AngularTargetVelocityMultiplier);
+			AngularVelocity * ControlComponent.ControlData.AngularTargetVelocityMultiplier);
 	}
 	else
 	{
 		Constraint->SetAngularDriveVelocityTarget(Chaos::FVec3(0));
 	}
 
+	// Danny TODO we don't actually have a shortcut for if the target is clean
+	ControlComponent.DirtyFlags &= ~ERigPhysicsControlComponentDirtyFlags::Target;
+
 	ControlRecord.PreviousTargetTM = TargetTM;
 	ControlRecord.PreviousTargetUpdateCounter = UpdateCounter;
 }
@@ -1618,7 +1664,7 @@ void FRigPhysicsSimulation::UpdateControlPrePhysics(
 // otherwise default to the physics element parent.
 void FRigPhysicsSimulation::UpdatePostPhysics(
 	URigHierarchy&              Hierarchy,
-	FRigPhysicsSolverComponent* SolverComponent, 
+	FRigPhysicsSolverComponent& SolverComponent, 
 	const float                 Alpha, 
 	const float                 DeltaTime,
 	const bool                  bUsingFixedStep)
@@ -1630,7 +1676,7 @@ void FRigPhysicsSimulation::UpdatePostPhysics(
 	{
 		return;
 	}
-	const FRigPhysicsSolverSettings& SolverSettings = SolverComponent->SolverSettings;
+	const FRigPhysicsSolverSettings& SolverSettings = SolverComponent.SolverSettings;
 	bool bGotInvalidSimulationData = false;
 
 	double PositionThresholdForResetSquared = SolverSettings.PositionThresholdForReset > 0 ?
@@ -1655,8 +1701,7 @@ void FRigPhysicsSimulation::UpdatePostPhysics(
 
 			// Calculate the target TM even if we're going to reset - it's likely useful for
 			// debugging (and this should be rare!)
-			Record.FinalComponentSpaceTM = ConvertSimSpaceTransformToComponentSpace(
-				SolverSettings, SimSpaceTM);
+			Record.FinalComponentSpaceTM = ConvertSimSpaceTransformToComponentSpace(SolverSettings, SimSpaceTM);
 
 			// Danny TODO rework this - currently in does the interpolation here, and stores it in
 			// the component, but I think the component should store the actual physics transform,
@@ -1667,11 +1712,10 @@ void FRigPhysicsSimulation::UpdatePostPhysics(
 				// This can cause joint separation. We probably want an option to blend
 				// in local (joint) space, perhaps splitting the alpha into orientation
 				// and position.
-				FTransform CurrentTM = Record.SourceComponentSpaceTM.ToTransform();
 				FQuat TargetQ = FQuat::Slerp(
-					CurrentTM.GetRotation(), Record.FinalComponentSpaceTM.GetRotation(), Alpha);
+					Record.SourceComponentSpaceTM.GetRotation(), Record.FinalComponentSpaceTM.GetRotation(), Alpha);
 				FVector TargetT = FMath::Lerp(
-					CurrentTM.GetTranslation(), Record.FinalComponentSpaceTM.GetTranslation(), Alpha);
+					Record.SourceComponentSpaceTM.GetTranslation(), Record.FinalComponentSpaceTM.GetTranslation(), Alpha);
 				Record.FinalComponentSpaceTM.SetRotation(TargetQ);
 				Record.FinalComponentSpaceTM.SetTranslation(TargetT);
 			}
@@ -1696,7 +1740,7 @@ void FRigPhysicsSimulation::UpdatePostPhysics(
 		UpdateCounter += 1;
 		// Set this to 3 since it gets decremented at the end of the update, and we need it to take
 		// effect at the start of the next update.
-		SolverComponent->TrackInputCounter = FMath::Max(SolverComponent->TrackInputCounter, 3);
+		SolverComponent.TrackInputCounter = FMath::Max(SolverComponent.TrackInputCounter, 3);
 	}
 
 	// If we found something invalid then we force the simulation to be as good as we can make it,
@@ -1723,9 +1767,8 @@ void FRigPhysicsSimulation::UpdatePostPhysics(
 					}
 					if (SourceKey.IsValid())
 					{
-						FTransform SourceComponentSpaceTM = Record.SourceComponentSpaceTM.ToTransform();
 						const FTransform SourceSimulationSpaceTM =
-							ConvertComponentSpaceTransformToSimSpace(SolverSettings, SourceComponentSpaceTM);
+							ConvertComponentSpaceTransformToSimSpace(SolverSettings, Record.SourceComponentSpaceTM);
 						Record.ActorHandle->SetWorldTransform(SourceSimulationSpaceTM);
 					}
 					Record.ActorHandle->SetLinearVelocity(FVector::ZeroVector);
@@ -1775,7 +1818,7 @@ void FRigPhysicsSimulation::StepSimulation(
 	const AActor*                OwningActorPtr,
 	const FRigVMExecuteContext&  ExecuteContext,
 	URigHierarchy&               Hierarchy,
-	FRigPhysicsSolverComponent*  SolverComponent,
+	FRigPhysicsSolverComponent&  SolverComponent,
 	const float                  DeltaTimeOverride,
 	const float                  SimulationSpaceDeltaTimeOverride,
 	const float                  Alpha)
@@ -1819,8 +1862,8 @@ void FRigPhysicsSimulation::StepSimulation(
 		return;
 	}
 
-	const FRigPhysicsSolverSettings& SolverSettings = SolverComponent->SolverSettings;
-	const FRigPhysicsSimulationSpaceSettings& SimulationSpaceSettings = SolverComponent->SimulationSpaceSettings;
+	const FRigPhysicsSolverSettings& SolverSettings = SolverComponent.SolverSettings;
+	const FRigPhysicsSimulationSpaceSettings& SimulationSpaceSettings = SolverComponent.SimulationSpaceSettings;
 
 	float FixedTimeStep = CVarControlRigPhysicsFixedTimeStepOverride.GetValueOnAnyThread() < 0 
 		? SolverSettings.FixedTimeStep : CVarControlRigPhysicsFixedTimeStepOverride.GetValueOnAnyThread();
@@ -1898,8 +1941,8 @@ void FRigPhysicsSimulation::StepSimulation(
 
 	UpdateWorldObjectsPostPhysics(World, SolverSettings, OwningActorPtr);
 
-	if (SolverComponent->TrackInputCounter > 0)
+	if (SolverComponent.TrackInputCounter > 0)
 	{
-		--SolverComponent->TrackInputCounter;
+		--SolverComponent.TrackInputCounter;
 	}
 }
diff --git a/Engine/Plugins/Experimental/ControlRigPhysics/Source/ControlRigPhysics/Private/RigPhysicsSimulation.h b/Engine/Plugins/Experimental/ControlRigPhysics/Source/ControlRigPhysics/Private/RigPhysicsSimulation.h
--- a/Engine/Plugins/Experimental/ControlRigPhysics/Source/ControlRigPhysics/Private/RigPhysicsSimulation.h
+++ b/Engine/Plugins/Experimental/ControlRigPhysics/Source/ControlRigPhysics/Private/RigPhysicsSimulation.h
@@ -70,20 +70,20 @@ struct FRigBodyRecord
 	// Cache the element key for where we will write the simulation result
 	FRigElementKey TargetElementKey;
 
-	// The final/simulated TM is stored before writing it into the output, so we can avoid
+	// The final/simulated transform is stored before writing it into the output, so we can avoid
 	// corrupting the output if anything is bad and we need to reset.
 	FTransform FinalComponentSpaceTM;
 
 	// These source (i.e. bone/element) transforms are updated for all records in UpdatePrePhysics.
 	// The times/validity are determined by the CurrentDeltaTime, PrevDeltaTime and update counters
 	// in the simulation itself.
-	UE::PhysicsControl::FPosQuat SourceComponentSpaceTM;
-	FVector                      SourceComponentSpaceVelocity = FVector::ZeroVector;
-	FVector                      SourceComponentSpaceAngularVelocity = FVector::ZeroVector;
+	FTransform SourceComponentSpaceTM;
+	FVector    SourceComponentSpaceVelocity = FVector::ZeroVector;
+	FVector    SourceComponentSpaceAngularVelocity = FVector::ZeroVector;
 
-	UE::PhysicsControl::FPosQuat PrevSourceComponentSpaceTM;
-	FVector                      PrevSourceComponentSpaceVelocity = FVector::ZeroVector;
-	FVector                      PrevSourceComponentSpaceAngularVelocity = FVector::ZeroVector;
+	FTransform PrevSourceComponentSpaceTM;
+	FVector    PrevSourceComponentSpaceVelocity = FVector::ZeroVector;
+	FVector    PrevSourceComponentSpaceAngularVelocity = FVector::ZeroVector;
 };
 
 //======================================================================================================================
@@ -105,7 +105,7 @@ struct FRigJointRecord
 	// Things that are updated as the simulation progresses
 
 	// The drive works with velocities so we store the previous target transform, and when it was stored.
-	UE::PhysicsControl::FPosQuat    PreviousDriveTargetTM;
+	FTransform                      PreviousDriveTargetTM;
 	// This is stored from the main solver update counter, marking when the previous drive TM was valid.
 	int64                           PreviousDriveTargetUpdateCounter = -999;
 };
@@ -128,10 +128,15 @@ struct FRigControlRecord
 
 	// Things that are updated as the simulation progresses
 
-	// The control works with velocities so we store the previous target transform, and when it was stored.
+	// The control works with velocities so we store the previous target, and when it was stored.
 	FTransform                      PreviousTargetTM;
 	// This is stored from the main solver update counter, marking when the previous target TM was valid.
 	int64                           PreviousTargetUpdateCounter = -999;
+
+	// When the FRigPhysicsControlComponent is dirty, this is calculated from the strengths etc, and
+	// then cached. Then if FRigPhysicsControlComponent is clean, we can see whether there is any
+	// strength/damping that would require the target to be updated.
+	bool bCachedHasStrengthOrDamping = true;
 };
 
 
@@ -150,7 +155,7 @@ public:
 	void Instantiate(
 		const FRigVMExecuteContext&       ExecuteContext, 
 		URigHierarchy&                    Hierarchy, 
-		const FRigPhysicsSolverComponent* SolverComponent);
+		const FRigPhysicsSolverComponent& SolverComponent);
 
 	// Integrates the simulation forwards.
 	// If DeltaTimeOverride is +ve, then that value is used.
@@ -161,7 +166,7 @@ public:
 		const AActor*                OwningActorPtr,
 		const FRigVMExecuteContext&  ExecuteContext, 
 		URigHierarchy&               Hierarchy,
-		FRigPhysicsSolverComponent*  SolverComponent,
+		FRigPhysicsSolverComponent&  SolverComponent,
 		const float                  DeltaTimeOverride, 
 		const float                  SimulationSpaceDeltaTimeOverride,
 		const float                  Alpha);
@@ -208,44 +213,46 @@ private:
 	};
 
 	// Creates the low level simulation
-	void InitialiseSimulation(const FRigPhysicsSolverComponent* SolverComponent);
+	void InitialiseSimulation();
 
 	// This collects all the bodies associated with the solver and makes records for them
-	void InitialiseBodyRecords(const URigHierarchy& Hierarchy, const FRigPhysicsSolverComponent* SolverComponent);
+	void InitialiseBodyRecords(const URigHierarchy& Hierarchy, const FRigPhysicsSolverComponent& SolverComponent);
 
 	// This collects all the joints associated with the solver and makes records for them
-	void InitialiseJointRecords(const URigHierarchy& Hierarchy, const FRigPhysicsSolverComponent* SolverComponent);
+	void InitialiseJointRecords(const URigHierarchy& Hierarchy, const FRigPhysicsSolverComponent& SolverComponent);
 
 	// This collects all the bodies associated with the solver and makes records for them
-	void InitialiseControlRecords(const URigHierarchy& Hierarchy, const FRigPhysicsSolverComponent* SolverComponent);
+	void InitialiseControlRecords(const URigHierarchy& Hierarchy, const FRigPhysicsSolverComponent& SolverComponent);
 
 	// Creates low level bodies
 	void InstantiatePhysicsBodies(
 		URigHierarchy&                    Hierarchy,
-		const FRigPhysicsSolverComponent* SolverComponent,
+		const FRigPhysicsSolverComponent& SolverComponent,
 		FRigPhysicsIgnorePairs&           IgnorePairs);
 
 	// Creates low level physics joints 
 	void InstantiatePhysicsJoints(
 		const URigHierarchy&              Hierarchy,
-		const FRigPhysicsSolverComponent* SolverComponent,
+		const FRigPhysicsSolverComponent& SolverComponent,
 		FRigPhysicsIgnorePairs&           IgnorePairs);
 
 	// Creates low level controls
 	void InstantiateControls(
 		const URigHierarchy&              Hierarchy,
-		const FRigPhysicsSolverComponent* SolverComponent,
+		const FRigPhysicsSolverComponent& SolverComponent,
 		FRigPhysicsIgnorePairs&           IgnorePairs);
 
 	// Creates collision shapes associated with the solver. Also applies IgnorePairs
 	void InstantiateSolverCollision(
-		const FRigPhysicsSolverComponent* SolverComponent,
+		const FRigPhysicsSolverComponent& SolverComponent,
 		FRigPhysicsIgnorePairs&           IgnorePairs);
 
 	// Returns true if the component is physics, and its solver matches the solver component
 	// (directly, or automatically)
 	bool ShouldComponentBeInSimulation(
-		const URigHierarchy& Hierarchy, FRigComponentKey SolverComponentKey, FRigComponentKey ComponentKey) const;
+		const URigHierarchy&              Hierarchy, 
+		const FRigPhysicsSolverComponent& SolverComponent, 
+		const FRigComponentKey&           ComponentKey) const;
 
 	// Creates an actor with collision. This will by dynamic if Dynamics is valid, or otherwise kinematic
 	ImmediatePhysics::FActorHandle* CreateBody(
@@ -255,51 +262,51 @@ private:
 		const FPhysicsControlModifierData* BodyData,
 		const FTransform&                  BodyRelSimSpaceTM) const;
 
-	// This releases the actors etc in the simulation, and destroys the simulation, but it doesn't
-	// change our records (apart from resetting pointers to actors etc).
+	// This releases the actors etc in the simulation, destroys the simulation, and resets our
+	// internal records.
 	void DestroyPhysicsSimulation();
 
 	void UpdatePrePhysics(
 		const FRigVMExecuteContext&        ExecuteContext, 
 		const URigHierarchy&               Hierarchy,
-		FRigPhysicsSolverComponent*        SolverComponent, 
+		FRigPhysicsSolverComponent&        SolverComponent, 
 		const float                        DeltaTime);
 
 	void UpdateBodyRecordPrePhysics(
 		const URigHierarchy&               Hierarchy,
-		const FRigPhysicsSolverComponent*  SolverComponent,
+		const FRigPhysicsSolverComponent&  SolverComponent,
 		const float                        DeltaTime,
 		FRigBodyRecord&                    Record,
 		const FRigPhysicsBodyComponent*    PhysicsComponent);
 
 	void UpdateBodyPrePhysics(
 		const FRigVMExecuteContext&        ExecuteContext, 
-		const FRigPhysicsSolverComponent*  SolverComponent,
+		const FRigPhysicsSolverComponent&  SolverComponent,
 		const FRigBodyRecord&              Record, 
-		const FRigPhysicsBodyComponent*    PhysicsComponent);
+		const FRigPhysicsBodyComponent&    PhysicsComponent);
 
 	void ApplyForcesPrePhysics(
 		const FRigVMExecuteContext&        ExecuteContext, 
-		const FRigPhysicsSolverComponent*  SolverComponent,
+		const FRigPhysicsSolverComponent&  SolverComponent,
 		const FRigBodyRecord&              Record, 
-		const FRigPhysicsBodyComponent*    PhysicsComponent);
+		const FRigPhysicsBodyComponent&    PhysicsComponent);
 
 	void UpdateJointPrePhysics(
 		const URigHierarchy&               Hierarchy, 
 		FRigJointRecord&                   JointRecord,
-		const FRigPhysicsJointComponent*   JointComponent, 
+		const FRigPhysicsJointComponent&   JointComponent, 
 		const float                        DeltaTime);
 
 	void UpdateControlPrePhysics(
 		FRigControlRecord&                 ControlRecord,
-		const FRigPhysicsControlComponent* ControlComponent,
-		const FRigPhysicsSolverComponent*  SolverComponent,
+		const FRigPhysicsControlComponent& ControlComponent,
+		const FRigPhysicsSolverComponent&  SolverComponent,
 		const URigHierarchy&               Hierarchy,
 		const float                        DeltaTime);
 
 	void UpdatePostPhysics(
 		URigHierarchy&              Hierarchy,
-		FRigPhysicsSolverComponent* SolverComponent, 
+		FRigPhysicsSolverComponent& SolverComponent, 
 		const float                 Alpha, 
 		const float                 DeltaTime,
 		const bool                  bUsingFixedStep);
@@ -316,16 +323,19 @@ private:
 	void InitSimulationSpace(const FTransform& ComponentToWorld, const FTransform& BoneToComponent);
 
 	void CheckForResetsPrePhysics(
-		const URigHierarchy& Hierarchy, FRigPhysicsSolverComponent* SolverComponent, const float DeltaTime);
+		const URigHierarchy& Hierarchy, FRigPhysicsSolverComponent& SolverComponent, const float DeltaTime);
 
 	// AbsoluteTime is used for calculating turbulence when applying the "wind"
 	FSimulationSpaceData UpdateSimulationSpaceStateAndCalculateData(
 		const FRigVMExecuteContext&         ExecuteContext,
 		const URigHierarchy&                Hierarchy,
-		const FRigPhysicsSolverComponent*   SolverComponent,
+		const FRigPhysicsSolverComponent&   SolverComponent,
 		const float                         Dt,
 		const double                        AbsoluteTime);
 
+	static FTransform GetSpaceTransform(
+		ERigPhysicsSimulationSpace Space, const FTransform& ComponentTM, const FTransform& BoneTM);
+
 	// Returns the simulation space transform, in world space
 	FTransform GetSimulationSpaceTransform(const FRigPhysicsSolverSettings& SolverSettings) const;
 
@@ -418,9 +428,8 @@ private:
 	// (b) it avoids passing it through the functions.
 	FSimulationSpaceData SimulationSpaceData;
 
-	// This is incremented at the end of each simulation step, used to identify when previously
-	// calculated values are valid
-	// Note that the universe will roll over before a uint64 does
+	// This is incremented at the very beginning of each simulation step, used to identify when previously
+	// calculated values are valid. Note that the universe will roll over before a int64 does
 	int64 UpdateCounter = 0;
 
 	// The update counter when PrevSourceComponentSpaceTM etc were written. Check this before using them
diff --git a/Engine/Plugins/Experimental/ControlRigPhysics/Source/ControlRigPhysics/Private/RigPhysicsSimulationSpace.cpp b/Engine/Plugins/Experimental/ControlRigPhysics/Source/ControlRigPhysics/Private/RigPhysicsSimulationSpace.cpp
--- a/Engine/Plugins/Experimental/ControlRigPhysics/Source/ControlRigPhysics/Private/RigPhysicsSimulationSpace.cpp
+++ b/Engine/Plugins/Experimental/ControlRigPhysics/Source/ControlRigPhysics/Private/RigPhysicsSimulationSpace.cpp
@@ -1,6 +1,8 @@
 // Copyright Epic Games, Inc. All Rights Reserved.
 
 #include "RigPhysicsSimulation.h"
+#include "RigPhysicsSimulationSpace.inl"
+
 #include "RigPhysicsSolverComponent.h"
 
 #include "Rigs/RigHierarchy.h"
@@ -30,187 +32,6 @@ TAutoConsoleVariable<int> CVarControlRigPhysicsShowSimulationSpaceInfo(
 	TEXT("ControlRig.Physics.ShowSimulationSpaceInfo"), 0,
 	TEXT("Shows information associated with the simulation space used in control rig physics."));
 
-//======================================================================================================================
-FORCEINLINE_DEBUGGABLE FTransform GetSpaceTransform(
-	ERigPhysicsSimulationSpace Space, const FTransform& ComponentTM, const FTransform& BoneTM)
-{
-	switch (Space)
-	{
-	case ERigPhysicsSimulationSpace::Component: return ComponentTM;
-	case ERigPhysicsSimulationSpace::World: return FTransform::Identity;
-	case ERigPhysicsSimulationSpace::SpaceBone: return BoneTM * ComponentTM;
-	default: ensureMsgf(false, TEXT("Unsupported Simulation Space")); return FTransform::Identity;
-	}
-}
-
-//======================================================================================================================
-FTransform FRigPhysicsSimulation::GetSimulationSpaceTransform(const FRigPhysicsSolverSettings& SolverSettings) const
-{
-	switch (SolverSettings.SimulationSpace)
-	{
-	case ERigPhysicsSimulationSpace::Component: return SimulationSpaceState.ComponentTM;
-	case ERigPhysicsSimulationSpace::World: return FTransform::Identity;
-	case ERigPhysicsSimulationSpace::SpaceBone: return SimulationSpaceState.BoneRelComponentTM * SimulationSpaceState.ComponentTM;
-	default: ensureMsgf(false, TEXT("Unsupported Simulation Space")); return FTransform::Identity;
-	}
-}
-
-//======================================================================================================================
-FTransform FRigPhysicsSimulation::ConvertComponentSpaceTransformToSimSpace(
-	const FRigPhysicsSolverSettings& SolverSettings, const FTransform& TM) const
-{
-	switch (SolverSettings.SimulationSpace)
-	{
-	case ERigPhysicsSimulationSpace::Component: return TM;
-	case ERigPhysicsSimulationSpace::World:  return TM * SimulationSpaceState.ComponentTM;
-	case ERigPhysicsSimulationSpace::SpaceBone: return TM.GetRelativeTransform(SimulationSpaceState.BoneRelComponentTM);
-	default: ensureMsgf(false, TEXT("Unsupported Simulation Space")); return TM;
-	}
-}
-
-//======================================================================================================================
-FVector FRigPhysicsSimulation::ConvertComponentSpacePositionToSimSpace(
-	const FRigPhysicsSolverSettings& SolverSettings, const FVector& P) const
-{
-	switch (SolverSettings.SimulationSpace)
-	{
-	case ERigPhysicsSimulationSpace::Component: return P;
-	case ERigPhysicsSimulationSpace::World:  return SimulationSpaceState.ComponentTM.TransformPositionNoScale(P);
-	case ERigPhysicsSimulationSpace::SpaceBone: return SimulationSpaceState.BoneRelComponentTM.InverseTransformPositionNoScale(P);
-	default: ensureMsgf(false, TEXT("Unsupported Simulation Space")); return P;
-	}
-}
-
-//======================================================================================================================
-FVector FRigPhysicsSimulation::ConvertComponentSpaceVectorToSimSpace(
-	const FRigPhysicsSolverSettings& SolverSettings, const FVector& V) const
-{
-	switch (SolverSettings.SimulationSpace)
-	{
-	case ERigPhysicsSimulationSpace::Component: return V;
-	case ERigPhysicsSimulationSpace::World:  return SimulationSpaceState.ComponentTM.TransformVectorNoScale(V);
-	case ERigPhysicsSimulationSpace::SpaceBone: return SimulationSpaceState.BoneRelComponentTM.InverseTransformVectorNoScale(V);
-	default: ensureMsgf(false, TEXT("Unsupported Simulation Space")); return V;
-	}
-}
-
-//======================================================================================================================
-FTransform FRigPhysicsSimulation::ConvertSimSpaceTransformToComponentSpace(
-	const FRigPhysicsSolverSettings& SolverSettings, const FTransform& TM) const
-{
-	switch (SolverSettings.SimulationSpace)
-	{
-	case ERigPhysicsSimulationSpace::Component: return TM;
-	case ERigPhysicsSimulationSpace::World:  return TM.GetRelativeTransform(SimulationSpaceState.ComponentTM);
-	case ERigPhysicsSimulationSpace::SpaceBone: return TM * SimulationSpaceState.BoneRelComponentTM;
-	default: ensureMsgf(false, TEXT("Unsupported Simulation Space")); return TM;
-	}
-}
-
-//======================================================================================================================
-FVector FRigPhysicsSimulation::ConvertSimSpacePositionToComponentSpace(
-	const FRigPhysicsSolverSettings& SolverSettings, const FVector& Position) const
-{
-	switch (SolverSettings.SimulationSpace)
-	{
-	case ERigPhysicsSimulationSpace::Component: return Position;
-	case ERigPhysicsSimulationSpace::World: return SimulationSpaceState.ComponentTM.InverseTransformPositionNoScale(Position);
-	case ERigPhysicsSimulationSpace::SpaceBone: return SimulationSpaceState.BoneRelComponentTM.TransformPositionNoScale(Position);
-	default: ensureMsgf(false, TEXT("Unsupported Simulation Space")); return Position;
-	}
-}
-
-//======================================================================================================================
-FVector FRigPhysicsSimulation::ConvertSimSpaceVectorToComponentSpace(
-	const FRigPhysicsSolverSettings& SolverSettings, const FVector& Vector) const
-{
-	switch (SolverSettings.SimulationSpace)
-	{
-	case ERigPhysicsSimulationSpace::Component: return Vector;
-	case ERigPhysicsSimulationSpace::World: return SimulationSpaceState.ComponentTM.InverseTransformVectorNoScale(Vector);
-	case ERigPhysicsSimulationSpace::SpaceBone: return SimulationSpaceState.BoneRelComponentTM.TransformVectorNoScale(Vector);
-	default: ensureMsgf(false, TEXT("Unsupported Simulation Space")); return Vector;
-	}
-}
-
-//======================================================================================================================
-FORCEINLINE_DEBUGGABLE FVector ConvertWorldVectorToSimSpaceNoScale(
-	ERigPhysicsSimulationSpace Space, const FVector& WorldVector, 
-	const FTransform& ComponentTM, const FTransform& BoneTM)
-{
-	switch (Space)
-	{
-	case ERigPhysicsSimulationSpace::Component: return ComponentTM.InverseTransformVectorNoScale(WorldVector);
-	case ERigPhysicsSimulationSpace::World: return WorldVector;
-	case ERigPhysicsSimulationSpace::SpaceBone:
-		return BoneTM.InverseTransformVectorNoScale(ComponentTM.InverseTransformVectorNoScale(WorldVector));
-	default: ensureMsgf(false, TEXT("Unsupported Simulation Space")); return WorldVector;
-	}
-}
-
-//======================================================================================================================
-FVector FRigPhysicsSimulation::ConvertWorldVectorToSimSpaceNoScale(
-	const FRigPhysicsSolverSettings& SolverSettings, const FVector& WorldVector) const
-{
-	switch (SolverSettings.SimulationSpace)
-	{
-	case ERigPhysicsSimulationSpace::Component: 
-		return SimulationSpaceState.ComponentTM.InverseTransformVectorNoScale(WorldVector);
-	case ERigPhysicsSimulationSpace::World: return WorldVector;
-	case ERigPhysicsSimulationSpace::SpaceBone:
-		return SimulationSpaceState.BoneRelComponentTM.InverseTransformVectorNoScale(
-			SimulationSpaceState.ComponentTM.InverseTransformVectorNoScale(WorldVector));
-	default: ensureMsgf(false, TEXT("Unsupported Simulation Space")); return WorldVector;
-	}
-}
-
-//======================================================================================================================
-FVector FRigPhysicsSimulation::ConvertWorldPositionToSimSpaceNoScale(
-	const FRigPhysicsSolverSettings& SolverSettings, const FVector& WorldPosition) const
-{
-	switch (SolverSettings.SimulationSpace)
-	{
-	case ERigPhysicsSimulationSpace::Component: 
-		return SimulationSpaceState.ComponentTM.InverseTransformPositionNoScale(WorldPosition);
-	case ERigPhysicsSimulationSpace::World: return WorldPosition;
-	case ERigPhysicsSimulationSpace::SpaceBone:
-		return SimulationSpaceState.BoneRelComponentTM.InverseTransformPositionNoScale(
-			SimulationSpaceState.ComponentTM.InverseTransformPositionNoScale(WorldPosition));
-	default: ensureMsgf(false, TEXT("Unsupported Simulation Space")); return WorldPosition;
-	}
-}
-
-//======================================================================================================================
-FTransform FRigPhysicsSimulation::ConvertWorldTransformToSimSpace(
-	const FRigPhysicsSolverSettings& SolverSettings, const FTransform& WorldTM) const
-{
-	switch (SolverSettings.SimulationSpace)
-	{
-	case ERigPhysicsSimulationSpace::Component: return WorldTM.GetRelativeTransform(SimulationSpaceState.ComponentTM);
-	case ERigPhysicsSimulationSpace::World: return WorldTM;
-	case ERigPhysicsSimulationSpace::SpaceBone:
-		return WorldTM.GetRelativeTransform(SimulationSpaceState.BoneRelComponentTM * SimulationSpaceState.ComponentTM);
-	default: ensureMsgf(false, TEXT("Unsupported Simulation Space")); return WorldTM;
-	}
-}
-
-//======================================================================================================================
-FTransform FRigPhysicsSimulation::ConvertCollisionSpaceTransformToSimSpace(
-	const FRigPhysicsSolverSettings& SolverSettings, const FTransform& TM) const
-{
-	FTransform SimSpaceTM = GetSpaceTransform(
-		SolverSettings.SimulationSpace, 
-		SimulationSpaceState.ComponentTM, 
-		SimulationSpaceState.BoneRelComponentTM);
-	FTransform CollisionSpaceTM = GetSpaceTransform(
-		SolverSettings.CollisionSpace, 
-		SimulationSpaceState.ComponentTM, 
-		SimulationSpaceState.BoneRelComponentTM);
-
-	FTransform WorldSpaceTM = TM * CollisionSpaceTM;
-	return WorldSpaceTM.GetRelativeTransform(SimSpaceTM);
-}
-
 //======================================================================================================================
 void FRigPhysicsSimulation::InitSimulationSpace(
 	const FTransform& ComponentTM,
@@ -220,7 +41,6 @@ void FRigPhysicsSimulation::InitSimulationSpace(
 	SimulationSpaceState.BoneRelComponentTM = BoneRelComponentTM;
 }
 
-
 //======================================================================================================================
 // Danny TODO We need support for double precision Perlin inputs - for now duplicate.
 namespace FPerlinHelpers
@@ -275,19 +95,17 @@ namespace FPerlinHelpers
 	}
 }
 
-
-
 //======================================================================================================================
 // Note - don't use the space conversion functions here as the state won't have been set yet.
 FRigPhysicsSimulation::FSimulationSpaceData FRigPhysicsSimulation::UpdateSimulationSpaceStateAndCalculateData(
 	const FRigVMExecuteContext&         ExecuteContext, 
 	const URigHierarchy&                Hierarchy,
-	const FRigPhysicsSolverComponent* SolverComponent,
+	const FRigPhysicsSolverComponent&   SolverComponent,
 	const float                         Dt,
 	const double                        AbsoluteTime)
 {
-	const FRigPhysicsSolverSettings& SolverSettings = SolverComponent->SolverSettings;
-	const FRigPhysicsSimulationSpaceSettings& SimulationSpaceSettings = SolverComponent->SimulationSpaceSettings;
+	const FRigPhysicsSolverSettings& SolverSettings = SolverComponent.SolverSettings;
+	const FRigPhysicsSimulationSpaceSettings& SimulationSpaceSettings = SolverComponent.SimulationSpaceSettings;
 
 	SimulationSpaceState.ComponentTM = ExecuteContext.GetToWorldSpaceTransform();
 
@@ -480,15 +298,15 @@ FRigPhysicsSimulation::FSimulationSpaceData FRigPhysicsSimulation::UpdateSimulat
 
 	// Transform world-space motion into simulation space Danny TODO note that this matches the code
 	// in RBAN, and is doing what the interface requires (i.e. movement of the space in the space of
-	// the space!). 
+	// the space!). Danny TODO note that we don't currently support scaling
 	SimulationSpaceData.LinearVelocity =
-		SimulationSpaceState.SimulationSpaceTM.InverseTransformVector(SimulationSpaceData.LinearVelocity);
+		SimulationSpaceState.SimulationSpaceTM.InverseTransformVectorNoScale(SimulationSpaceData.LinearVelocity);
 	SimulationSpaceData.AngularVelocity =
-		SimulationSpaceState.SimulationSpaceTM.InverseTransformVector(SimulationSpaceData.AngularVelocity);
+		SimulationSpaceState.SimulationSpaceTM.InverseTransformVectorNoScale(SimulationSpaceData.AngularVelocity);
 	SimulationSpaceData.LinearAcceleration =
-		SimulationSpaceState.SimulationSpaceTM.InverseTransformVector(SimulationSpaceData.LinearAcceleration);
+		SimulationSpaceState.SimulationSpaceTM.InverseTransformVectorNoScale(SimulationSpaceData.LinearAcceleration);
 	SimulationSpaceData.AngularAcceleration =
-		SimulationSpaceState.SimulationSpaceTM.InverseTransformVector(SimulationSpaceData.AngularAcceleration);
+		SimulationSpaceState.SimulationSpaceTM.InverseTransformVectorNoScale(SimulationSpaceData.AngularAcceleration);
 
 	if (GEngine && CVarControlRigPhysicsShowSimulationSpaceInfo.GetValueOnAnyThread())
 	{
diff --git a/Engine/Plugins/Experimental/ControlRigPhysics/Source/ControlRigPhysics/Private/RigPhysicsSimulationSpace.inl b/Engine/Plugins/Experimental/ControlRigPhysics/Source/ControlRigPhysics/Private/RigPhysicsSimulationSpace.inl
new file mode 100644
--- /dev/null
+++ b/Engine/Plugins/Experimental/ControlRigPhysics/Source/ControlRigPhysics/Private/RigPhysicsSimulationSpace.inl
@@ -0,0 +1,186 @@
+// Copyright Epic Games, Inc. All Rights Reserved.
+
+#pragma once
+
+#include "RigPhysicsSimulation.h"
+#include "PhysicsControlHelpers.h"
+
+#include "Engine/Engine.h"
+
+//======================================================================================================================
+FORCEINLINE FTransform FRigPhysicsSimulation::GetSpaceTransform(
+	ERigPhysicsSimulationSpace Space, const FTransform& ComponentTM, const FTransform& BoneTM)
+{
+	switch (Space)
+	{
+	case ERigPhysicsSimulationSpace::Component: return ComponentTM;
+	case ERigPhysicsSimulationSpace::World: return FTransform::Identity;
+	case ERigPhysicsSimulationSpace::SpaceBone: return BoneTM * ComponentTM;
+	default: ensureMsgf(false, TEXT("Unsupported Simulation Space")); return FTransform::Identity;
+	}
+}
+
+//======================================================================================================================
+FORCEINLINE FTransform FRigPhysicsSimulation::GetSimulationSpaceTransform(const FRigPhysicsSolverSettings& SolverSettings) const
+{
+	switch (SolverSettings.SimulationSpace)
+	{
+	case ERigPhysicsSimulationSpace::Component: return SimulationSpaceState.ComponentTM;
+	case ERigPhysicsSimulationSpace::World: return FTransform::Identity;
+	case ERigPhysicsSimulationSpace::SpaceBone: return 
+		SimulationSpaceState.BoneRelComponentTM * SimulationSpaceState.ComponentTM;
+	default: ensureMsgf(false, TEXT("Unsupported Simulation Space")); return FTransform::Identity;
+	}
+}
+
+//======================================================================================================================
+FORCEINLINE FTransform FRigPhysicsSimulation::ConvertComponentSpaceTransformToSimSpace(
+	const FRigPhysicsSolverSettings& SolverSettings, const FTransform& TM) const
+{
+	switch (SolverSettings.SimulationSpace)
+	{
+	case ERigPhysicsSimulationSpace::Component: return TM;
+	case ERigPhysicsSimulationSpace::World: return TM * SimulationSpaceState.ComponentTM;
+	case ERigPhysicsSimulationSpace::SpaceBone: return TM.GetRelativeTransform(SimulationSpaceState.BoneRelComponentTM);
+	default: ensureMsgf(false, TEXT("Unsupported Simulation Space")); return TM;
+	}
+}
+
+//======================================================================================================================
+FORCEINLINE FVector FRigPhysicsSimulation::ConvertComponentSpacePositionToSimSpace(
+	const FRigPhysicsSolverSettings& SolverSettings, const FVector& P) const
+{
+	switch (SolverSettings.SimulationSpace)
+	{
+	case ERigPhysicsSimulationSpace::Component: return P;
+	case ERigPhysicsSimulationSpace::World:  return SimulationSpaceState.ComponentTM.TransformPositionNoScale(P);
+	case ERigPhysicsSimulationSpace::SpaceBone: return SimulationSpaceState.BoneRelComponentTM.InverseTransformPositionNoScale(P);
+	default: ensureMsgf(false, TEXT("Unsupported Simulation Space")); return P;
+	}
+}
+
+//======================================================================================================================
+FORCEINLINE FVector FRigPhysicsSimulation::ConvertComponentSpaceVectorToSimSpace(
+	const FRigPhysicsSolverSettings& SolverSettings, const FVector& V) const
+{
+	switch (SolverSettings.SimulationSpace)
+	{
+	case ERigPhysicsSimulationSpace::Component: return V;
+	case ERigPhysicsSimulationSpace::World: return SimulationSpaceState.ComponentTM.TransformVectorNoScale(V);
+	case ERigPhysicsSimulationSpace::SpaceBone: 
+		return SimulationSpaceState.BoneRelComponentTM.InverseTransformVectorNoScale(V);
+	default: ensureMsgf(false, TEXT("Unsupported Simulation Space")); return V;
+	}
+}
+
+//======================================================================================================================
+FORCEINLINE FTransform FRigPhysicsSimulation::ConvertSimSpaceTransformToComponentSpace(
+	const FRigPhysicsSolverSettings& SolverSettings, const FTransform& TM) const
+{
+	switch (SolverSettings.SimulationSpace)
+	{
+	case ERigPhysicsSimulationSpace::Component: return TM;
+	case ERigPhysicsSimulationSpace::World: return TM.GetRelativeTransform(SimulationSpaceState.ComponentTM);
+	case ERigPhysicsSimulationSpace::SpaceBone: return TM * SimulationSpaceState.BoneRelComponentTM;
+	default: ensureMsgf(false, TEXT("Unsupported Simulation Space")); return TM;
+	}
+}
+
+//======================================================================================================================
+FORCEINLINE FVector FRigPhysicsSimulation::ConvertSimSpacePositionToComponentSpace(
+	const FRigPhysicsSolverSettings& SolverSettings, const FVector& Position) const
+{
+	switch (SolverSettings.SimulationSpace)
+	{
+	case ERigPhysicsSimulationSpace::Component: return Position;
+	case ERigPhysicsSimulationSpace::World: return SimulationSpaceState.ComponentTM.InverseTransformPositionNoScale(Position);
+	case ERigPhysicsSimulationSpace::SpaceBone: return SimulationSpaceState.BoneRelComponentTM.TransformPositionNoScale(Position);
+	default: ensureMsgf(false, TEXT("Unsupported Simulation Space")); return Position;
+	}
+}
+
+//======================================================================================================================
+FORCEINLINE FVector FRigPhysicsSimulation::ConvertSimSpaceVectorToComponentSpace(
+	const FRigPhysicsSolverSettings& SolverSettings, const FVector& Vector) const
+{
+	switch (SolverSettings.SimulationSpace)
+	{
+	case ERigPhysicsSimulationSpace::Component: return Vector;
+	case ERigPhysicsSimulationSpace::World: return SimulationSpaceState.ComponentTM.InverseTransformVectorNoScale(Vector);
+	case ERigPhysicsSimulationSpace::SpaceBone: return SimulationSpaceState.BoneRelComponentTM.TransformVectorNoScale(Vector);
+	default: ensureMsgf(false, TEXT("Unsupported Simulation Space")); return Vector;
+	}
+}
+
+//======================================================================================================================
+FORCEINLINE FVector ConvertWorldVectorToSimSpaceNoScale(
+	ERigPhysicsSimulationSpace Space, const FVector& WorldVector, 
+	const FTransform& ComponentTM, const FTransform& BoneTM)
+{
+	switch (Space)
+	{
+	case ERigPhysicsSimulationSpace::Component: return ComponentTM.InverseTransformVectorNoScale(WorldVector);
+	case ERigPhysicsSimulationSpace::World: return WorldVector;
+	case ERigPhysicsSimulationSpace::SpaceBone:
+		return BoneTM.InverseTransformVectorNoScale(ComponentTM.InverseTransformVectorNoScale(WorldVector));
+	default: ensureMsgf(false, TEXT("Unsupported Simulation Space")); return WorldVector;
+	}
+}
+
+//======================================================================================================================
+FORCEINLINE FVector FRigPhysicsSimulation::ConvertWorldVectorToSimSpaceNoScale(
+	const FRigPhysicsSolverSettings& SolverSettings, const FVector& WorldVector) const
+{
+	switch (SolverSettings.SimulationSpace)
+	{
+	case ERigPhysicsSimulationSpace::Component: 
+		return SimulationSpaceState.ComponentTM.InverseTransformVectorNoScale(WorldVector);
+	case ERigPhysicsSimulationSpace::World: return WorldVector;
+	case ERigPhysicsSimulationSpace::SpaceBone:
+		return SimulationSpaceState.BoneRelComponentTM.InverseTransformVectorNoScale(
+			SimulationSpaceState.ComponentTM.InverseTransformVectorNoScale(WorldVector));
+	default: ensureMsgf(false, TEXT("Unsupported Simulation Space")); return WorldVector;
+	}
+}
+
+//======================================================================================================================
+FORCEINLINE FVector FRigPhysicsSimulation::ConvertWorldPositionToSimSpaceNoScale(
+	const FRigPhysicsSolverSettings& SolverSettings, const FVector& WorldPosition) const
+{
+	switch (SolverSettings.SimulationSpace)
+	{
+	case ERigPhysicsSimulationSpace::Component: 
+		return SimulationSpaceState.ComponentTM.InverseTransformPositionNoScale(WorldPosition);
+	case ERigPhysicsSimulationSpace::World: return WorldPosition;
+	case ERigPhysicsSimulationSpace::SpaceBone:
+		return SimulationSpaceState.BoneRelComponentTM.InverseTransformPositionNoScale(
+			SimulationSpaceState.ComponentTM.InverseTransformPositionNoScale(WorldPosition));
+	default: ensureMsgf(false, TEXT("Unsupported Simulation Space")); return WorldPosition;
+	}
+}
+
+//======================================================================================================================
+FORCEINLINE FTransform FRigPhysicsSimulation::ConvertWorldTransformToSimSpace(
+	const FRigPhysicsSolverSettings& SolverSettings, const FTransform& WorldTM) const
+{
+	switch (SolverSettings.SimulationSpace)
+	{
+	case ERigPhysicsSimulationSpace::Component: return WorldTM.GetRelativeTransform(SimulationSpaceState.ComponentTM);
+	case ERigPhysicsSimulationSpace::World: return WorldTM;
+	case ERigPhysicsSimulationSpace::SpaceBone:
+		return WorldTM.GetRelativeTransform(SimulationSpaceState.BoneRelComponentTM * SimulationSpaceState.ComponentTM);
+	default: ensureMsgf(false, TEXT("Unsupported Simulation Space")); return WorldTM;
+	}
+}
+
+//======================================================================================================================
+FORCEINLINE FTransform FRigPhysicsSimulation::ConvertCollisionSpaceTransformToSimSpace(
+	const FRigPhysicsSolverSettings& SolverSettings, const FTransform& TM) const
+{
+	FTransform SimSpaceTM = GetSpaceTransform(
+		SolverSettings.SimulationSpace, SimulationSpaceState.ComponentTM, SimulationSpaceState.BoneRelComponentTM);
+	FTransform CollisionSpaceTM = GetSpaceTransform(
+		SolverSettings.CollisionSpace, SimulationSpaceState.ComponentTM, SimulationSpaceState.BoneRelComponentTM);
+	FTransform WorldSpaceTM = TM * CollisionSpaceTM;
+	return WorldSpaceTM.GetRelativeTransform(SimSpaceTM);
+}
diff --git a/Engine/Plugins/Experimental/ControlRigPhysics/Source/ControlRigPhysics/Private/RigPhysicsSimulationVisualization.cpp b/Engine/Plugins/Experimental/ControlRigPhysics/Source/ControlRigPhysics/Private/RigPhysicsSimulationVisualization.cpp
--- a/Engine/Plugins/Experimental/ControlRigPhysics/Source/ControlRigPhysics/Private/RigPhysicsSimulationVisualization.cpp
+++ b/Engine/Plugins/Experimental/ControlRigPhysics/Source/ControlRigPhysics/Private/RigPhysicsSimulationVisualization.cpp
@@ -1,6 +1,7 @@
 // Copyright Epic Games, Inc. All Rights Reserved.
 
 #include "RigPhysicsSimulation.h"
+#include "RigPhysicsSimulationSpace.inl"
 
 #include "Rigs/RigHierarchy.h"
 #include "Rigs/RigHierarchyController.h"
diff --git a/Engine/Plugins/Experimental/ControlRigPhysics/Source/ControlRigPhysics/Private/RigPhysicsSimulationWorldObjects.cpp b/Engine/Plugins/Experimental/ControlRigPhysics/Source/ControlRigPhysics/Private/RigPhysicsSimulationWorldObjects.cpp
--- a/Engine/Plugins/Experimental/ControlRigPhysics/Source/ControlRigPhysics/Private/RigPhysicsSimulationWorldObjects.cpp
+++ b/Engine/Plugins/Experimental/ControlRigPhysics/Source/ControlRigPhysics/Private/RigPhysicsSimulationWorldObjects.cpp
@@ -1,6 +1,7 @@
 // Copyright Epic Games, Inc. All Rights Reserved.
 
 #include "RigPhysicsSimulation.h"
+#include "RigPhysicsSimulationSpace.inl"
 
 #include "Physics/ImmediatePhysics/ImmediatePhysicsActorHandle.h"
 #include "Physics/ImmediatePhysics/ImmediatePhysicsAdapters.h"
@@ -13,10 +14,16 @@
 #include "CollisionQueryParams.h"
 #include "CollisionShape.h"
 
+#include "Stats/Stats.h"
+#include "Stats/StatsHierarchical.h"
+
 //======================================================================================================================
 // Note that the simulation mutex will already be taken when we are called
 void FRigPhysicsSimulation::UpdateWorldObjectsPrePhysics(const FRigPhysicsSolverSettings& SolverSettings)
 {
+	DECLARE_SCOPE_HIERARCHICAL_COUNTER_FUNC()
+	QUICK_SCOPE_CYCLE_COUNTER(STAT_RigPhysics_UpdateWorldObjectsPrePhysics);
+
 	if (SolverSettings.WorldCollisionType == ERigPhysicsWorldCollisionType::None)
 	{
 		for (TPair<uint32, FWorldObject>& WorldObjectPair : *WorldObjects.Get())
@@ -63,6 +70,8 @@ void FRigPhysicsSimulation::UpdateWorldObjectsPrePhysics(const FRigPhysicsSolver
 	{
 		FWorldObject& WorldObject = WorldObjectPair.Value;
 		FTransform SimSpaceTM = ConvertWorldTransformToSimSpace(SolverSettings, WorldObject.ComponentWorldTransform);
+		// Note that the code above should have ensured we don't have a null pointer
+		check(WorldObject.ActorHandle != nullptr);
 		if (WorldObject.LastSeenUpdateCounter == -1)
 		{
 			UE_LOG(LogRigPhysics, Log,
@@ -79,7 +88,6 @@ void FRigPhysicsSimulation::UpdateWorldObjectsPrePhysics(const FRigPhysicsSolver
 			WorldObject.ActorHandle->SetKinematicTarget(SimSpaceTM);
 		}
 	}
-
 }
 
 //======================================================================================================================
@@ -87,6 +95,9 @@ void FRigPhysicsSimulation::UpdateWorldObjectsPrePhysics(const FRigPhysicsSolver
 void FRigPhysicsSimulation::UpdateWorldObjectsPostPhysics(
 	const UWorld* World, const FRigPhysicsSolverSettings& SolverSettings, const AActor* OwningActorPtr)
 {
+	DECLARE_SCOPE_HIERARCHICAL_COUNTER_FUNC()
+	QUICK_SCOPE_CYCLE_COUNTER(STAT_RigPhysics_UpdateWorldObjectsPostPhysics);
+
 	WorldOverlapBox = FBox();
 	if (SolverSettings.WorldCollisionType == ERigPhysicsWorldCollisionType::None)
 	{
diff --git a/Engine/Plugins/Experimental/ControlRigPhysics/Source/ControlRigPhysics/Public/RigPhysicsControlComponent.h b/Engine/Plugins/Experimental/ControlRigPhysics/Source/ControlRigPhysics/Public/RigPhysicsControlComponent.h
--- a/Engine/Plugins/Experimental/ControlRigPhysics/Source/ControlRigPhysics/Public/RigPhysicsControlComponent.h
+++ b/Engine/Plugins/Experimental/ControlRigPhysics/Source/ControlRigPhysics/Public/RigPhysicsControlComponent.h
@@ -11,6 +11,16 @@
 
 #define UE_API CONTROLRIGPHYSICS_API
 
+enum class ERigPhysicsControlComponentDirtyFlags : uint8
+{
+	None = 0,
+	Data = 1 << 0, // Strength etc (not target)
+	Target = 1 << 1,
+	All = Data | Target
+};
+ENUM_CLASS_FLAGS(ERigPhysicsControlComponentDirtyFlags);
+
+
 /**
  * A component that can be added to hierarchy elements (joints) to add the data required to control 
  * the simulation of them
@@ -58,6 +68,9 @@ public:
 	UPROPERTY(BlueprintReadWrite, EditAnywhere, Category = Control)
 	FPhysicsControlTarget ControlTarget;
 
+	// Flags indicating what has been changed and needs to be set in the physics representation
+	mutable ERigPhysicsControlComponentDirtyFlags DirtyFlags = ERigPhysicsControlComponentDirtyFlags::All;
+
 	UE_API virtual void Save(FArchive& A) override;
 	UE_API virtual void Load(FArchive& Ar) override;
 
diff --git a/Engine/Plugins/Experimental/ControlRigPhysics/Source/ControlRigPhysics/Public/RigPhysicsJointComponent.h b/Engine/Plugins/Experimental/ControlRigPhysics/Source/ControlRigPhysics/Public/RigPhysicsJointComponent.h
--- a/Engine/Plugins/Experimental/ControlRigPhysics/Source/ControlRigPhysics/Public/RigPhysicsJointComponent.h
+++ b/Engine/Plugins/Experimental/ControlRigPhysics/Source/ControlRigPhysics/Public/RigPhysicsJointComponent.h
@@ -14,6 +14,15 @@
 
 #define UE_API CONTROLRIGPHYSICS_API
 
+enum class ERigPhysicsJointComponentDirtyFlags : uint8
+{
+	None  = 0,
+	Joint = 1 << 0,
+	Drive = 1 << 1,
+	All   = Joint | Drive
+};
+ENUM_CLASS_FLAGS(ERigPhysicsJointComponentDirtyFlags);
+
 USTRUCT(BlueprintType)
 struct FRigPhysicsJointComponent : public FRigBaseComponent
 {
@@ -47,6 +56,9 @@ public:
 	UPROPERTY(EditAnywhere, BlueprintReadWrite, Category = Physics, meta = (ShowOnlyInnerProperties))
 	FRigPhysicsDriveData DriveData;
 
+	// Flags indicating what has been changed and needs to be set in the physics representation
+	mutable ERigPhysicsJointComponentDirtyFlags DirtyFlags = ERigPhysicsJointComponentDirtyFlags::All;
+
 	UE_API virtual void Save(FArchive& Ar) override;
 	UE_API virtual void Load(FArchive& Ar) override;
 
diff --git a/Engine/Plugins/Experimental/PhysicsControl/Source/PhysicsControl/Private/RigidBodyWithControl.cpp b/Engine/Plugins/Experimental/PhysicsControl/Source/PhysicsControl/Private/RigidBodyWithControl.cpp
--- a/Engine/Plugins/Experimental/PhysicsControl/Source/PhysicsControl/Private/RigidBodyWithControl.cpp
+++ b/Engine/Plugins/Experimental/PhysicsControl/Source/PhysicsControl/Private/RigidBodyWithControl.cpp
@@ -30,7 +30,6 @@ DECLARE_CYCLE_STAT(TEXT("RigidBodyNodeWithControl_ApplyControlsAndModifiers"), S
 DECLARE_CYCLE_STAT(TEXT("RigidBodyNodeWithControl_ApplyForcesAndTorques"), STAT_RigidBodyNodeWithControl_ApplyForcesAndTorques, STATGROUP_Anim);
 DECLARE_CYCLE_STAT(TEXT("RigidBodyNodeWithControl_ApplyControls"), STAT_RigidBodyNodeWithControl_ApplyControls, STATGROUP_Anim);
 DECLARE_CYCLE_STAT(TEXT("RigidBodyNodeWithControl_ApplyModifiers"), STAT_RigidBodyNodeWithControl_ApplyModifiers, STATGROUP_Anim);
-DECLARE_CYCLE_STAT(TEXT("RigidBodyNodeWithControl_SetPhysicsJointEnabled"), STAT_RigidBodyNodeWithControl_SetPhysicsJointEnabled, STATGROUP_Anim);
 DECLARE_CYCLE_STAT(TEXT("RigidBodyNodeWithControl_ApplyKinematicTargets"), STAT_RigidBodyNodeWithControl_ApplyKinematicTargets, STATGROUP_Anim);
 DECLARE_CYCLE_STAT(TEXT("RigidBodyNodeWithControl_ApplyCurrentConstraintProfile"), STAT_RigidBodyNodeWithControl_ApplyCurrentConstraintProfile, STATGROUP_Anim);
 
@@ -38,25 +37,20 @@ constexpr int32 ConstraintChildIndex = 0;
 constexpr int32 ConstraintParentIndex = 1;
 
 //======================================================================================================================
-void SetPhysicsJointEnabled(ImmediatePhysics::FJointHandle* const JointHandle, const bool bIsEnabled)
+void SetPhysicsJointEnabled(ImmediatePhysics::FJointHandle& JointHandle, const bool bIsEnabled)
 {
-	SCOPE_CYCLE_COUNTER(STAT_RigidBodyNodeWithControl_SetPhysicsJointEnabled);
-
-	if (JointHandle)
+	if (ImmediatePhysics::FJointHandle::FChaosConstraintHandle* ConstraintHandle = JointHandle.GetConstraint())
 	{
-		if (ImmediatePhysics::FJointHandle::FChaosConstraintHandle* ConstraintHandle = JointHandle->GetConstraint())
-		{
-			ConstraintHandle->SetConstraintEnabled(bIsEnabled);
+		ConstraintHandle->SetConstraintEnabled(bIsEnabled);
 
-			if (!bIsEnabled)
-			{
-				// The call above disables the constraint, but if any actor in the simulation is
-				// flagged as dirty, it gets re-enabled! See UE-204006 TODO remove this workaround
-				// when the bug has been fixed.
-				ConstraintHandle->SetDriveParams(
-					Chaos::FVec3::ZeroVector, Chaos::FVec3::ZeroVector, Chaos::FVec3::ZeroVector, 
-					Chaos::FVec3::ZeroVector, Chaos::FVec3::ZeroVector, Chaos::FVec3::ZeroVector);
-			}
+		if (!bIsEnabled)
+		{
+			// The call above disables the constraint, but if any actor in the simulation is
+			// flagged as dirty, it gets re-enabled! See UE-204006 TODO remove this workaround
+			// when the bug has been fixed.
+			ConstraintHandle->SetDriveParams(
+				Chaos::FVec3::ZeroVector, Chaos::FVec3::ZeroVector, Chaos::FVec3::ZeroVector, 
+				Chaos::FVec3::ZeroVector, Chaos::FVec3::ZeroVector, Chaos::FVec3::ZeroVector);
 		}
 	}
 }
@@ -141,16 +135,16 @@ FName MapConstraintsBehaviorTypeToString(const MapConstraintsBehaviorType InType
 
 //======================================================================================================================
 void UpdateBodyFromModifierData(
-	ImmediatePhysics::FActorHandle*    ActorHandle,
-	ImmediatePhysics::FSimulation*     PhysicsSimulation,
+	ImmediatePhysics::FActorHandle&    ActorHandle,
+	ImmediatePhysics::FSimulation&     PhysicsSimulation,
 	const FPhysicsControlModifierData& ModifierData,
-	const FVector& SimSpaceGravity)
+	const FVector&                     SimSpaceGravity)
 {
 	if (ModifierData.MovementType != EPhysicsMovementType::Default)
 	{
 		// Note that there's an early out if there's no change needed, so this should be OK.
-		PhysicsSimulation->SetIsKinematic(
-			ActorHandle,
+		PhysicsSimulation.SetIsKinematic(
+			&ActorHandle,
 			ModifierData.MovementType != EPhysicsMovementType::Simulated);
 	}
 
@@ -158,70 +152,68 @@ void UpdateBodyFromModifierData(
 	// for all kinematics whether or not they were under a modifier.
 
 	// Scale gravity
-	float GravityMultiplier = ModifierData.GravityMultiplier;
-	if (GravityMultiplier != 1 && ActorHandle->IsGravityEnabled())
+	const float GravityMultiplier = ModifierData.GravityMultiplier;
+	if (GravityMultiplier != 1 && ActorHandle.IsGravityEnabled())
 	{
-		float Mass = (float)ActorHandle->GetMass();
-		FVector AntiGravityForce = SimSpaceGravity * (-Mass * (1.0f - GravityMultiplier));
-		ActorHandle->AddForce(AntiGravityForce);
+		const float Mass = (float)ActorHandle.GetMass();
+		const FVector AntiGravityForce = SimSpaceGravity * (Mass * (GravityMultiplier - 1.0f));
+		ActorHandle.AddForce(AntiGravityForce);
 	}
 
-	// Set collision
-	PhysicsSimulation->SetHasCollision(
-		ActorHandle, CollisionEnabledHasPhysics(ModifierData.CollisionType));
+	// Set collision. This has an early out, and doesn't flag actors as dirty etc if nothing has changed.
+	PhysicsSimulation.SetHasCollision(
+		&ActorHandle, CollisionEnabledHasPhysics(ModifierData.CollisionType));
 
-	ActorHandle->SetUpdateKinematicFromSimulation(ModifierData.bUpdateKinematicFromSimulation);
+	// This simply sets a flag - doesn't trigger any work to call it when it hasn't changed.
+	ActorHandle.SetUpdateKinematicFromSimulation(ModifierData.bUpdateKinematicFromSimulation);
 }
 
 //======================================================================================================================
 ImmediatePhysics::FJointHandle* CreatePhysicsJoint(
-	ImmediatePhysics::FSimulation*  Simulation,
-	ImmediatePhysics::FActorHandle* ChildActorHandle,
-	ImmediatePhysics::FActorHandle* ParentActorHandle)
+	ImmediatePhysics::FSimulation&  Simulation,
+	ImmediatePhysics::FActorHandle& ChildActorHandle,
+	ImmediatePhysics::FActorHandle& ParentActorHandle)
 {
 	ImmediatePhysics::FJointHandle* JointHandle = nullptr;
 
-	if (Simulation && ChildActorHandle && ParentActorHandle)
-	{
-		Chaos::FPBDJointSettings Settings;
+	Chaos::FPBDJointSettings Settings;
 
-		Settings.LinearMotionTypes = {
-			Chaos::EJointMotionType::Free, Chaos::EJointMotionType::Free, Chaos::EJointMotionType::Free };
-		Settings.AngularMotionTypes = {
-			Chaos::EJointMotionType::Free, Chaos::EJointMotionType::Free, Chaos::EJointMotionType::Free };
+	Settings.LinearMotionTypes = {
+		Chaos::EJointMotionType::Free, Chaos::EJointMotionType::Free, Chaos::EJointMotionType::Free };
+	Settings.AngularMotionTypes = {
+		Chaos::EJointMotionType::Free, Chaos::EJointMotionType::Free, Chaos::EJointMotionType::Free };
 
-		Settings.bLinearPositionDriveEnabled = { true, true, true };
-		Settings.bLinearVelocityDriveEnabled = { true, true, true };
-		Settings.LinearDriveForceMode = Chaos::EJointForceMode::Acceleration;
+	Settings.bLinearPositionDriveEnabled = { true, true, true };
+	Settings.bLinearVelocityDriveEnabled = { true, true, true };
+	Settings.LinearDriveForceMode = Chaos::EJointForceMode::Acceleration;
 
-		Settings.bAngularSLerpPositionDriveEnabled = true;
-		Settings.bAngularSLerpVelocityDriveEnabled = true;
+	Settings.bAngularSLerpPositionDriveEnabled = true;
+	Settings.bAngularSLerpVelocityDriveEnabled = true;
 
-		Settings.bAngularTwistPositionDriveEnabled = false;
-		Settings.bAngularTwistVelocityDriveEnabled = false;
-		Settings.bAngularSwingPositionDriveEnabled = false;
-		Settings.bAngularSwingVelocityDriveEnabled = false;
-		Settings.AngularDriveForceMode = Chaos::EJointForceMode::Acceleration;
+	Settings.bAngularTwistPositionDriveEnabled = false;
+	Settings.bAngularTwistVelocityDriveEnabled = false;
+	Settings.bAngularSwingPositionDriveEnabled = false;
+	Settings.bAngularSwingVelocityDriveEnabled = false;
+	Settings.AngularDriveForceMode = Chaos::EJointForceMode::Acceleration;
 
-		// For control, we shouldn't be in situations where mass conditioning is needed.
-		Settings.bMassConditioningEnabled = false;
+	// For control, we shouldn't be in situations where mass conditioning is needed.
+	Settings.bMassConditioningEnabled = false;
 
-		Settings.bUseLinearSolver = true;
-		// It's not our job to change collision settings - that should come from the physics asset.
-		// However, the naming of this is unclear - if collisions are disabled in the physics asset,
-		// trust that this doesn't enable them.
-		Settings.bCollisionEnabled = true;
+	Settings.bUseLinearSolver = true;
+	// It's not our job to change collision settings - that should come from the physics asset.
+	// However, the naming of this is unclear - if collisions are disabled in the physics asset,
+	// trust that this doesn't enable them.
+	Settings.bCollisionEnabled = true;
 
-		FVector ChildCoMPositionOffset = ChildActorHandle->GetLocalCoMTransform().GetLocation();
-		Settings.ConnectorTransforms[ConstraintChildIndex].SetLocation(ChildCoMPositionOffset);
+	FVector ChildCoMPositionOffset = ChildActorHandle.GetLocalCoMTransform().GetLocation();
+	Settings.ConnectorTransforms[ConstraintChildIndex].SetLocation(ChildCoMPositionOffset);
 
-		Settings.Sanitize();
-		JointHandle = Simulation->CreateJoint(ImmediatePhysics::FJointSetup(Settings, ChildActorHandle, ParentActorHandle));
-	}
+	Settings.Sanitize();
+	JointHandle = Simulation.CreateJoint(ImmediatePhysics::FJointSetup(Settings, &ChildActorHandle, &ParentActorHandle));
 
 	if (JointHandle)
 	{
-		SetPhysicsJointEnabled(JointHandle, false); // Disable constraints by default.
+		SetPhysicsJointEnabled(*JointHandle, false); // Disable constraints by default.
 	}
 
 	return JointHandle;
@@ -320,7 +312,10 @@ bool FAnimNode_RigidBodyWithControl::CreateNamedControl(
 		ParentBodyHandle = FindBodyFromBoneName(ParentBoneName);
 	}
 	ImmediatePhysics::FActorHandle* const ChildBodyHandle = FindBodyFromBoneName(ChildBoneName);
-	JointHandle = CreatePhysicsJoint(Simulation, ChildBodyHandle, ParentBodyHandle);
+	if (Simulation && ChildBodyHandle && ParentBodyHandle)
+	{
+		JointHandle = CreatePhysicsJoint(*Simulation, *ChildBodyHandle, *ParentBodyHandle);
+	}
 
 	if (!JointHandle)
 	{
@@ -567,15 +562,17 @@ static void ConvertStrengthToSpringParams(
 }
 
 //======================================================================================================================
-// Adjusts the constraint spring drive settings to reflect the control data
-// Returns true if there is some control, false if the drive has no effect
 bool UpdateDriveSpringDamperSettings(
-	ImmediatePhysics::FJointHandle*   JointHandle,
+	ImmediatePhysics::FJointHandle&   JointHandle,
 	const Chaos::FPBDJointSettings&   Settings, 
 	const FPhysicsControlData&        Data,
 	const FPhysicsControlMultiplier&  Multiplier)
 {
-	Chaos::FPBDJointConstraintHandle* Constraint = JointHandle->GetConstraint();
+	Chaos::FPBDJointConstraintHandle* Constraint = JointHandle.GetConstraint();
+	if (!Constraint)
+	{
+		return false;
+	}
 	float AngularSpring;
 	float AngularDamping;
 	const float MaxTorque = Data.MaxTorque * Multiplier.MaxTorqueMultiplier;
@@ -618,11 +615,14 @@ bool UpdateDriveSpringDamperSettings(
 
 	Constraint->SetDriveParams(
 		LinearSpring, LinearDamping, MaxForce,
-		Chaos::FVec3(AngularSpring), Chaos::FVec3(AngularDamping), Chaos::FVec3(Data.MaxTorque));
+		Chaos::FVec3(AngularSpring), Chaos::FVec3(AngularDamping), Chaos::FVec3(MaxTorque));
 
 	bool bHaveAngular = (AngularSpring + AngularDamping) > 0;
 	bool bHaveLinear = (LinearSpring + LinearDamping).GetMax() > 0;
-	return bHaveLinear || bHaveAngular;
+	bool bHaveAnything = bHaveLinear || bHaveAngular;
+
+	SetPhysicsJointEnabled(JointHandle, bHaveAnything);
+	return bHaveAnything;
 }
 
 //======================================================================================================================
@@ -681,7 +681,7 @@ void FAnimNode_RigidBodyWithControl::ApplyControl(FRigidBodyControlRecord& Contr
 
 			const Chaos::FPBDJointSettings& JointSettings = Constraint->GetSettings();
 			if (UpdateDriveSpringDamperSettings(
-				JointHandle, JointSettings, ControlRecord.ControlData, ControlRecord.ControlMultiplier))
+				*JointHandle, JointSettings, ControlRecord.ControlData, ControlRecord.ControlMultiplier))
 			{
 				const FActorHandle* ChildActorHandle = JointHandle->GetActorHandles()[ConstraintChildIndex];
 				const FActorHandle* ParentActorHandle = JointHandle->GetActorHandles()[ConstraintParentIndex];
@@ -767,7 +767,7 @@ void FAnimNode_RigidBodyWithControl::ApplyModifier(
 	if (BodyModifierRecord.ActorHandle && Simulation)
 	{
 		UpdateBodyFromModifierData(
-			BodyModifierRecord.ActorHandle, Simulation, BodyModifierRecord.ModifierData, SimSpaceGravity);
+			*BodyModifierRecord.ActorHandle, *Simulation, BodyModifierRecord.ModifierData, SimSpaceGravity);
 	}
 }
 
@@ -903,7 +903,15 @@ void FAnimNode_RigidBodyWithControl::ApplyControlsAndModifiers(const FVector& Si
 				{
 					ApplyControl(ControlRecord, DeltaTime);
 				}
-				SetPhysicsJointEnabled(ControlRecord.JointHandle, ControlRecord.IsEnabled());
+				else
+				{
+					if (ControlRecord.JointHandle)
+					{
+						SetPhysicsJointEnabled(*ControlRecord.JointHandle, false);
+					}
+					// Ensure the control doesn't use a bogus velocity calculation when it is re-enabled
+					ControlRecord.ExpectedUpdateCounter.Reset();
+				}
 			}
 		}
 
diff --git a/Engine/Plugins/Experimental/PhysicsControl/Source/PhysicsControl/Public/AnimNode_RigidBodyWithControl.h b/Engine/Plugins/Experimental/PhysicsControl/Source/PhysicsControl/Public/AnimNode_RigidBodyWithControl.h
--- a/Engine/Plugins/Experimental/PhysicsControl/Source/PhysicsControl/Public/AnimNode_RigidBodyWithControl.h
+++ b/Engine/Plugins/Experimental/PhysicsControl/Source/PhysicsControl/Public/AnimNode_RigidBodyWithControl.h
@@ -57,25 +57,79 @@ enum class MapConstraintsBehaviorType : uint8
  * RigidBodyWithControl does.
  */
 PHYSICSCONTROL_API void UpdateBodyFromModifierData(
-	ImmediatePhysics::FActorHandle* ActorHandle,
-	ImmediatePhysics::FSimulation* PhysicsSimulation,
+	ImmediatePhysics::FActorHandle&    ActorHandle,
+	ImmediatePhysics::FSimulation&     PhysicsSimulation,
 	const FPhysicsControlModifierData& ModifierData,
-	const FVector& SimSpaceGravity);
+	const FVector&                     SimSpaceGravity);
 
+// Adjusts the constraint spring drive settings to reflect the control data. Also enabled or
+// disables the underlying constraint based on whether there is any strength/damping. Returns true
+// if there is some control, false if the drive has no effect
 PHYSICSCONTROL_API bool UpdateDriveSpringDamperSettings(
-	ImmediatePhysics::FJointHandle* JointHandle,
-	const Chaos::FPBDJointSettings& Settings,
-	const FPhysicsControlData& Data,
+	ImmediatePhysics::FJointHandle&  JointHandle,
+	const Chaos::FPBDJointSettings&  Settings,
+	const FPhysicsControlData&       Data,
 	const FPhysicsControlMultiplier& Multiplier);
 
 // Creates a physics joint in the form that's useful for controls.
 PHYSICSCONTROL_API ImmediatePhysics::FJointHandle* CreatePhysicsJoint(
-	ImmediatePhysics::FSimulation*  Simulation,
-	ImmediatePhysics::FActorHandle* ChildActorHandle,
-	ImmediatePhysics::FActorHandle* ParentActorHandle);
+	ImmediatePhysics::FSimulation&  Simulation,
+	ImmediatePhysics::FActorHandle& ChildActorHandle,
+	ImmediatePhysics::FActorHandle& ParentActorHandle);
 
 PHYSICSCONTROL_API void SetPhysicsJointEnabled(
-	ImmediatePhysics::FJointHandle* const JointHandle, const bool bIsEnabled);
+	ImmediatePhysics::FJointHandle& JointHandle, const bool bIsEnabled);
+
+// Deprecated versions of the functions above
+
+UE_DEPRECATED(5.8, "Use the version of UpdateBodyFromModifierData taking references instead instead")
+PHYSICSCONTROL_API inline void UpdateBodyFromModifierData(
+	ImmediatePhysics::FActorHandle*    ActorHandle,
+	ImmediatePhysics::FSimulation*     PhysicsSimulation,
+	const FPhysicsControlModifierData& ModifierData,
+	const FVector&                     SimSpaceGravity)
+{
+	if (ActorHandle && PhysicsSimulation)
+	{
+		UpdateBodyFromModifierData(*ActorHandle, *PhysicsSimulation, ModifierData, SimSpaceGravity);
+	}
+}
+
+UE_DEPRECATED(5.8, "Use the version of UpdateDriveSpringDamperSettings taking references instead instead")
+PHYSICSCONTROL_API inline bool UpdateDriveSpringDamperSettings(
+	ImmediatePhysics::FJointHandle* JointHandle,
+	const Chaos::FPBDJointSettings& Settings,
+	const FPhysicsControlData& Data,
+	const FPhysicsControlMultiplier& Multiplier)
+{
+	if (JointHandle)
+	{
+		return UpdateDriveSpringDamperSettings(*JointHandle, Settings, Data, Multiplier);
+	}
+	return false;
+}
+
+UE_DEPRECATED(5.8, "Use the version of CreatePhysicsJoint taking references instead instead")
+PHYSICSCONTROL_API inline ImmediatePhysics::FJointHandle* CreatePhysicsJoint(
+	ImmediatePhysics::FSimulation* Simulation,
+	ImmediatePhysics::FActorHandle* ChildActorHandle,
+	ImmediatePhysics::FActorHandle* ParentActorHandle)
+{
+	if (Simulation && ChildActorHandle && ParentActorHandle)
+	{
+		return CreatePhysicsJoint(*Simulation, *ChildActorHandle, *ParentActorHandle);
+	}
+	return nullptr;
+}
+
+UE_DEPRECATED(5.8, "Use the version of SetPhysicsJointEnabled taking references instead instead")
+PHYSICSCONTROL_API inline void SetPhysicsJointEnabled(ImmediatePhysics::FJointHandle* JointHandle, const bool bIsEnabled)
+{
+	if (JointHandle)
+	{
+		SetPhysicsJointEnabled(*JointHandle, bIsEnabled);
+	}
+}
 
 /**
  * Controller that simulates physics based on the physics asset of the skeletal mesh component
