Control Rig Physics: Ensures that world objects get initialised in the collision system properly - fixes problems with objects not colliding unless they have moved

#Synced-CL 51405397

diff --git a/Engine/Plugins/Experimental/ControlRigPhysics/Source/ControlRigPhysics/Private/RigPhysicsSimulationWorldObjects.cpp b/Engine/Plugins/Experimental/ControlRigPhysics/Source/ControlRigPhysics/Private/RigPhysicsSimulationWorldObjects.cpp
--- a/Engine/Plugins/Experimental/ControlRigPhysics/Source/ControlRigPhysics/Private/RigPhysicsSimulationWorldObjects.cpp
+++ b/Engine/Plugins/Experimental/ControlRigPhysics/Source/ControlRigPhysics/Private/RigPhysicsSimulationWorldObjects.cpp
@@ -78,7 +78,10 @@ void FRigPhysicsSimulation::UpdateWorldObjectsPrePhysics(const FRigPhysicsSolver
 				TEXT("Control rig %s - Adding object %s"), *OwnerName.ToString(),
 				*WorldObject.ActorHandle->GetName().ToString());
 
-			WorldObject.ActorHandle->SetWorldTransform(SimSpaceTM);
+			// We need to call InitWorldTransform (not just SetWorldTransform) here, to ensure the
+			// velocities are set to zero (though they should be after creation anyway), and that
+			// the object has its bounds set.
+			WorldObject.ActorHandle->InitWorldTransform(SimSpaceTM);
 			// If the object is found in the next overlap, then its update counter will be updated.
 			// However, in case it's not, we set it here too to make sure that it expires. 
 			WorldObject.LastSeenUpdateCounter = UpdateCounter;
diff --git a/Engine/Source/Runtime/Engine/Public/Physics/ImmediatePhysics/ImmediatePhysicsChaos/ImmediatePhysicsActorHandle_Chaos.h b/Engine/Source/Runtime/Engine/Public/Physics/ImmediatePhysics/ImmediatePhysicsChaos/ImmediatePhysicsActorHandle_Chaos.h
--- a/Engine/Source/Runtime/Engine/Public/Physics/ImmediatePhysics/ImmediatePhysicsChaos/ImmediatePhysicsActorHandle_Chaos.h
+++ b/Engine/Source/Runtime/Engine/Public/Physics/ImmediatePhysics/ImmediatePhysicsChaos/ImmediatePhysicsActorHandle_Chaos.h
@@ -66,7 +66,10 @@ namespace ImmediatePhysics_Chaos
 		 */
 		ENGINE_API void SetHasCollision(bool bHasCollision);
 
-		/** Sets the world transform, zeroes velocity, etc.*/
+		/**
+		 * Sets the world transform, zeroes velocity, etc. This also initialises the bounds, which may be
+		 * important since otherwise they would only get updated if the object actually moves.
+		 */
 		ENGINE_API void InitWorldTransform(const FTransform& WorldTM);
 
 		/** Sets the world transform, maintains velocity etc.*/
