Control Rig Physics: Fix serialization of shapes (name and ContributeToMass)

#Synced-CL 51405397

diff --git a/Engine/Plugins/Experimental/ControlRigPhysics/Source/ControlRigPhysics/Private/RigPhysicsData.cpp b/Engine/Plugins/Experimental/ControlRigPhysics/Source/ControlRigPhysics/Private/RigPhysicsData.cpp
--- a/Engine/Plugins/Experimental/ControlRigPhysics/Source/ControlRigPhysics/Private/RigPhysicsData.cpp
+++ b/Engine/Plugins/Experimental/ControlRigPhysics/Source/ControlRigPhysics/Private/RigPhysicsData.cpp
@@ -5,8 +5,54 @@
 #include "PhysicsControlHelpers.h"
 #include "Chaos/ChaosConstraintSettings.h"
 
+
+//======================================================================================================================
+FArchive& operator <<(FArchive& Ar, FRigPhysicsCollisionShape& Data)
+{
+	Ar << Data.RestOffset;
+	Ar << Data.Name;
+	Ar << Data.bContributeToMass;
+	return Ar;
+}
+
+//======================================================================================================================
+FArchive& operator <<(FArchive& Ar, FRigPhysicsCollisionBox& Data)
+{
+	if (Ar.CustomVer(FPhysicsControlObjectVersion::GUID) >= FPhysicsControlObjectVersion::ControlRigShapeData)
+	{
+		Ar << static_cast<FRigPhysicsCollisionShape&>(Data);
+	}
+	Ar << Data.TM;
+	Ar << Data.Extents;
+	return Ar;
+}
+
 //======================================================================================================================
+FArchive& operator <<(FArchive& Ar, FRigPhysicsCollisionSphere& Data)
+{
+	if (Ar.CustomVer(FPhysicsControlObjectVersion::GUID) >= FPhysicsControlObjectVersion::ControlRigShapeData)
+	{
+		Ar << static_cast<FRigPhysicsCollisionShape&>(Data);
+	}
+	Ar << Data.TM;
+	Ar << Data.Radius;
+	return Ar;
+}
 
+//======================================================================================================================
+FArchive& operator <<(FArchive& Ar, FRigPhysicsCollisionCapsule& Data)
+{
+	if (Ar.CustomVer(FPhysicsControlObjectVersion::GUID) >= FPhysicsControlObjectVersion::ControlRigShapeData)
+	{
+		Ar << static_cast<FRigPhysicsCollisionShape&>(Data);
+	}
+	Ar << Data.TM;
+	Ar << Data.Radius;
+	Ar << Data.Length;
+	return Ar;
+}
+
+//======================================================================================================================
 #include UE_INLINE_GENERATED_CPP_BY_NAME(RigPhysicsData)
 FArchive& ArchiveConstraintBaseParams(FArchive& Ar, FConstraintBaseParams& Data)
 {
diff --git a/Engine/Plugins/Experimental/ControlRigPhysics/Source/ControlRigPhysics/Public/RigPhysicsData.h b/Engine/Plugins/Experimental/ControlRigPhysics/Source/ControlRigPhysics/Public/RigPhysicsData.h
--- a/Engine/Plugins/Experimental/ControlRigPhysics/Source/ControlRigPhysics/Public/RigPhysicsData.h
+++ b/Engine/Plugins/Experimental/ControlRigPhysics/Source/ControlRigPhysics/Public/RigPhysicsData.h
@@ -117,20 +117,20 @@ struct FRigPhysicsCollisionShape
 {
 	GENERATED_BODY()
 
-	/** Offset used when generating contact points. This allows you to smooth out
-	the Minkowski sum by radius R. Useful for making objects slide smoothly
-	on top of irregularities  */
+	// Offset used when generating contact points. This allows you to smooth out the Minkowski sum
+	// by radius R. Useful for making objects slide smoothly on top of irregularities
 	UPROPERTY(Category = Shape, EditAnywhere, meta = (ClampMin = "0.0"))
 	float RestOffset = 0.0f;
 
 	UPROPERTY(Category = Shape, EditAnywhere)
 	FName Name;
 
-	/** True if this shape should contribute to the overall mass of the body it
-		belongs to. This lets you create extra collision volumes which do not affect
-		the mass properties of an object. */
+	// True if this shape should contribute to the overall mass of the body it belongs to. This lets
+	// you create extra ollision volumes which do not affect the mass properties of an object.
 	UPROPERTY(Category = Shape, EditAnywhere)
 	bool bContributeToMass = true;
+
+	friend FArchive& operator <<(FArchive& Ar, FRigPhysicsCollisionShape& Data);
 };
 
 /**
@@ -145,19 +145,14 @@ struct FRigPhysicsCollisionBox : public FRigPhysicsCollisionShape
 		TM(InTM), Extents(InExtents) {}
 	FRigPhysicsCollisionBox() {}
 
-	friend FArchive& operator <<(FArchive& Ar, FRigPhysicsCollisionBox& Data)
-	{
-		Ar << Data.TM;
-		Ar << Data.Extents;
-		return Ar;
-	}
-
 	UPROPERTY(EditAnywhere, Category = Shapes)
 	FTransform TM;
 
 	// These are the full extents of the box in each axis
 	UPROPERTY(EditAnywhere, Category = Shapes)
 	FVector Extents = FVector::OneVector * RigPhysicsDefaultShapeSize;
+
+	friend FArchive& operator <<(FArchive& Ar, FRigPhysicsCollisionBox& Data);
 };
 
 /**
@@ -173,18 +168,13 @@ struct FRigPhysicsCollisionSphere : public FRigPhysicsCollisionShape
 	}
 	FRigPhysicsCollisionSphere() {}
 
-	friend FArchive& operator <<(FArchive& Ar, FRigPhysicsCollisionSphere& Data)
-	{
-		Ar << Data.TM;
-		Ar << Data.Radius;
-		return Ar;
-	}
-
 	UPROPERTY(EditAnywhere, Category = Shapes)
 	FTransform TM;
 
 	UPROPERTY(EditAnywhere, Category = Shapes, meta = (ClampMin = "0.0"))
 	float Radius = RigPhysicsDefaultShapeSize;
+
+	friend FArchive& operator <<(FArchive& Ar, FRigPhysicsCollisionSphere& Data);
 };
 
 /**
@@ -199,14 +189,6 @@ struct FRigPhysicsCollisionCapsule : public FRigPhysicsCollisionShape
 		: TM(InTM), Radius(InRadius), Length(InLength) {}
 	FRigPhysicsCollisionCapsule() {}
 
-	friend FArchive& operator <<(FArchive& Ar, FRigPhysicsCollisionCapsule& Data)
-	{
-		Ar << Data.TM;
-		Ar << Data.Radius;
-		Ar << Data.Length;
-		return Ar;
-	}
-
 	UPROPERTY(EditAnywhere, Category = Shapes)
 	FTransform TM;
 
@@ -216,6 +198,8 @@ struct FRigPhysicsCollisionCapsule : public FRigPhysicsCollisionShape
 	// This is the length of the core part of the capsule. The total length will be Length + 2 * Radius
 	UPROPERTY(EditAnywhere, Category = Shapes, meta = (ClampMin = "0.0"))
 	float Length = RigPhysicsDefaultShapeSize;
+
+	friend FArchive& operator <<(FArchive& Ar, FRigPhysicsCollisionCapsule& Data);
 };
 
 // How to combine friction/restitution values. 
diff --git a/Engine/Plugins/Experimental/PhysicsControl/Source/PhysicsControl/Public/PhysicsControlObjectVersion.h b/Engine/Plugins/Experimental/PhysicsControl/Source/PhysicsControl/Public/PhysicsControlObjectVersion.h
--- a/Engine/Plugins/Experimental/PhysicsControl/Source/PhysicsControl/Public/PhysicsControlObjectVersion.h
+++ b/Engine/Plugins/Experimental/PhysicsControl/Source/PhysicsControl/Public/PhysicsControlObjectVersion.h
@@ -39,8 +39,9 @@ struct FPhysicsControlObjectVersion
 		ControlRigJointEnabled,
 		RemoveSkeletalAnimationVelocityMultiplier,
 		RemoveUseSkeletalAnimation,
+		ControlRigShapeData,
 
-		// -----<new versions can be added above this line>-------------------------------------------------
+		// -----<new versions can be added above this line>-----
 		VersionPlusOne,
 		LatestVersion = VersionPlusOne - 1
 	};
