Enable Control Rig Physics plugin by default

#Synced-CL 51539364

diff --git a/Engine/Plugins/Experimental/ControlRigPhysics/ControlRigPhysics.uplugin b/Engine/Plugins/Experimental/ControlRigPhysics/ControlRigPhysics.uplugin
--- a/Engine/Plugins/Experimental/ControlRigPhysics/ControlRigPhysics.uplugin
+++ b/Engine/Plugins/Experimental/ControlRigPhysics/ControlRigPhysics.uplugin
@@ -15,6 +15,7 @@
 	"IsBetaVersion": false,
 	"IsExperimentalVersion": true,
 	"Installed": false,
+	"EnabledByDefault": true,
 	"Plugins": [
 		{
 			"Name": "ControlRig",
@@ -23,13 +24,13 @@
 		{
 			"Name": "PhysicsControl",
 			"Enabled": true
-		},
+		}
 	],
 	"Modules": [
 		{
 			"Name": "ControlRigPhysics",
 			"Type": "Runtime",
 			"LoadingPhase": "Default"
-		},
+		}
 	]
 }
\ No newline at end of file
