[Backout] - CL51568916
#rnx
Original CL Desc
-----------------------------------------------------------------
Enable Control Rig Physics plugin by default

diff --git a/Engine/Plugins/Experimental/ControlRigPhysics/ControlRigPhysics.uplugin b/Engine/Plugins/Experimental/ControlRigPhysics/ControlRigPhysics.uplugin
--- a/Engine/Plugins/Experimental/ControlRigPhysics/ControlRigPhysics.uplugin
+++ b/Engine/Plugins/Experimental/ControlRigPhysics/ControlRigPhysics.uplugin
@@ -15,7 +15,6 @@
 	"IsBetaVersion": false,
 	"IsExperimentalVersion": true,
 	"Installed": false,
-	"EnabledByDefault": true,
 	"Plugins": [
 		{
 			"Name": "ControlRig",
@@ -24,13 +23,13 @@
 		{
 			"Name": "PhysicsControl",
 			"Enabled": true
-		}
+		},
 	],
 	"Modules": [
 		{
 			"Name": "ControlRigPhysics",
 			"Type": "Runtime",
 			"LoadingPhase": "Default"
-		}
+		},
 	]
 }
\ No newline at end of file
