Control Rig Physics: Clarify the solver friction and restitution defaults

#Synced-CL 52003484

diff --git a/Engine/Plugins/Experimental/ControlRigPhysics/Source/ControlRigPhysics/Public/RigPhysicsExecution.h b/Engine/Plugins/Experimental/ControlRigPhysics/Source/ControlRigPhysics/Public/RigPhysicsExecution.h
--- a/Engine/Plugins/Experimental/ControlRigPhysics/Source/ControlRigPhysics/Public/RigPhysicsExecution.h
+++ b/Engine/Plugins/Experimental/ControlRigPhysics/Source/ControlRigPhysics/Public/RigPhysicsExecution.h
@@ -44,9 +44,12 @@ struct FRigUnit_AddPhysicsSolver : public FRigUnit_PhysicsBaseMutable
 	{
 		Owner.Type = ERigElementType::Bone;
 		// Default the material here to have friction and restitution. Then the interactions are
-		// easily adjusted on the dynamic bodies.
+		// easily adjusted on the dynamic bodies, because we prefer to combine by multiplying. If
+		// restitution was zero here, then setting restitution on the dynamic bodies would have no effect.
 		SolverSettings.Collision.Material.Friction = 1.0f;
 		SolverSettings.Collision.Material.Restitution = 1.0f;
+		SolverSettings.Collision.Material.FrictionCombineMode = ERigPhysicsCombineMode::Multiply;
+		SolverSettings.Collision.Material.RestitutionCombineMode = ERigPhysicsCombineMode::Multiply;
 	}
 
 	RIGVM_METHOD()
