fix for CRP to correctly add world bodies that are owned by instanced static mesh components
#Synced-CL 52056680

diff --git a/Engine/Plugins/Experimental/ControlRigPhysics/Source/ControlRigPhysics/Private/RigPhysicsSimulation.h b/Engine/Plugins/Experimental/ControlRigPhysics/Source/ControlRigPhysics/Private/RigPhysicsSimulation.h
--- a/Engine/Plugins/Experimental/ControlRigPhysics/Source/ControlRigPhysics/Private/RigPhysicsSimulation.h
+++ b/Engine/Plugins/Experimental/ControlRigPhysics/Source/ControlRigPhysics/Private/RigPhysicsSimulation.h
@@ -461,6 +461,10 @@ private:
 		// should expire/forget about it.
 		int64 LastSeenUpdateCounter;
 
+		// For instanced static mesh components, the index of the specific instance this actor
+		// represents. INDEX_NONE for non-instanced components.
+		int32 InstanceBodyIndex = INDEX_NONE;
+
 		bool GetExpired(int64 InUpdateCounter, int64 InRetainCount) const
 		{
 			if (LastSeenUpdateCounter != -1 && InUpdateCounter > LastSeenUpdateCounter + InRetainCount)
diff --git a/Engine/Plugins/Experimental/ControlRigPhysics/Source/ControlRigPhysics/Private/RigPhysicsSimulationWorldObjects.cpp b/Engine/Plugins/Experimental/ControlRigPhysics/Source/ControlRigPhysics/Private/RigPhysicsSimulationWorldObjects.cpp
--- a/Engine/Plugins/Experimental/ControlRigPhysics/Source/ControlRigPhysics/Private/RigPhysicsSimulationWorldObjects.cpp
+++ b/Engine/Plugins/Experimental/ControlRigPhysics/Source/ControlRigPhysics/Private/RigPhysicsSimulationWorldObjects.cpp
@@ -17,6 +17,8 @@
 #include "Stats/Stats.h"
 #include "Stats/StatsHierarchical.h"
 
+#include "Components/InstancedStaticMeshComponent.h"
+
 //======================================================================================================================
 // Note that the simulation mutex will already be taken when we are called
 void FRigPhysicsSimulation::UpdateWorldObjectsPrePhysics(const FRigPhysicsSolverSettings& SolverSettings)
@@ -196,16 +198,54 @@ void FRigPhysicsSimulation::UpdateWorldObjectsPostPhysics(
 							{
 								if (UPrimitiveComponent* OverlapComp = Overlap.GetComponent())
 								{
-									uint32 Key = OverlapComp->GetUniqueID();
+									// For instanced static mesh components, each overlap result
+									// corresponds to a specific instance via Overlap.ItemIndex.
+									// The component-level BodyInstance has null BodySetup and no
+									// physics actor, so we must use the per-instance body instead.
+									int32 OverlapInstanceIndex = INDEX_NONE;
+									if (OverlapComp->IsA<UInstancedStaticMeshComponent>())
+									{
+										OverlapInstanceIndex = Overlap.ItemIndex;
+									}
+
+									// Use a per-instance composite key for ISM so each instance
+									// gets its own FWorldObject entry.
+									const uint32 Key = (OverlapInstanceIndex == INDEX_NONE)
+										? OverlapComp->GetUniqueID()
+										: HashCombine(OverlapComp->GetUniqueID(), (uint32)OverlapInstanceIndex);
+
 									FWorldObject* WorldObject = WorldObjects->Find(Key);
 									if (!WorldObject)
 									{
 										// New object - add it to the sim
 										if (OverlapComp->GetOwner() && OverlapComp->GetOwner() != OwningActorPtrCopy)
 										{
+											// Resolve the correct BodyInstance and transform.
+											// For ISM the component-level BodyInstance has null
+											// BodySetup, so use the per-instance body.
+											const FBodyInstance* BodyInstanceToUse = &OverlapComp->BodyInstance;
+
+											if (OverlapInstanceIndex != INDEX_NONE)
+											{
+												UInstancedStaticMeshComponent* ISMComp = Cast<UInstancedStaticMeshComponent>(OverlapComp);
+												BodyInstanceToUse = ISMComp	
+													? ISMComp->GetBodyInstance(NAME_None, false, OverlapInstanceIndex)
+													: nullptr;
+												if (BodyInstanceToUse == nullptr)
+												{
+													// Instance has no physics body
+													continue;
+												}
+											}
+
+											// Chaos holds scale separately to the world transform of the body
+											FTransform TransformToUse = BodyInstanceToUse->GetUnrealWorldTransform();
+											TransformToUse.SetScale3D(BodyInstanceToUse->Scale3D);
+
 											WorldObject = &WorldObjects->Add(Key, FWorldObject());
 
 											WorldObject->WorldPrimitiveComponent = OverlapComp;
+											WorldObject->InstanceBodyIndex = OverlapInstanceIndex;
 
 											// Note that we need to create the simulation actor here, even
 											// though we're in the game thread, so we can copy data out of the
@@ -214,15 +254,14 @@ void FRigPhysicsSimulation::UpdateWorldObjectsPostPhysics(
 
 											WorldObject->ActorHandle = SimulationCopy->CreateActor(
 												ImmediatePhysics::MakeKinematicActorSetup(
-													&OverlapComp->BodyInstance, OverlapComp->GetComponentTransform()));
+													BodyInstanceToUse, TransformToUse));
 
 #if WITH_EDITOR
 											WorldObject->ActorHandle->SetName(*OverlapComp->GetOwner()->GetActorLabel());
 #endif
 											SimulationCopy->AddToCollidingPairs(WorldObject->ActorHandle);
 
-											WorldObject->ComponentWorldTransform = 
-												OverlapComp->BodyInstance.GetUnrealWorldTransform();
+											WorldObject->ComponentWorldTransform = TransformToUse;
 
 											// Flag that the simulation TM needs to be set, rather
 											// than using a kinematic target.
@@ -231,11 +270,30 @@ void FRigPhysicsSimulation::UpdateWorldObjectsPostPhysics(
 									}
 									else
 									{
-										if (TStrongObjectPtr<UPrimitiveComponent> WPC = 
+										if (TStrongObjectPtr<UPrimitiveComponent> WPC =
 											WorldObject->WorldPrimitiveComponent.Pin())
 										{
-											WorldObject->ComponentWorldTransform = 
-												WPC->BodyInstance.GetUnrealWorldTransform();
+											FBodyInstance* BodyInstanceToUse = &WPC->BodyInstance;
+
+											if (WorldObject->InstanceBodyIndex != INDEX_NONE)
+											{
+												UInstancedStaticMeshComponent* ISMComp = Cast<UInstancedStaticMeshComponent>(WPC.Get());
+												BodyInstanceToUse = ISMComp 
+													? ISMComp->GetBodyInstance(NAME_None, false, WorldObject->InstanceBodyIndex)
+													: nullptr;
+
+												if (BodyInstanceToUse == nullptr)
+												{
+													// Instance has no physics body
+													continue;
+												}
+											}
+
+											// Chaos holds scale separately to the world transform
+											FTransform TransformToUse = BodyInstanceToUse->GetUnrealWorldTransform();
+											TransformToUse.SetScale3D(BodyInstanceToUse->Scale3D);
+
+											WorldObject->ComponentWorldTransform = TransformToUse;
 											WorldObject->LastSeenUpdateCounter = UpdateCounterCopy;
 										}
 									}
