Control Rig Physics - no API/functional changes here - just splitting a large file, and reducing unnecessary includes

#Synced-CL 52074282

diff --git a/Engine/Plugins/Experimental/ControlRigPhysics/Source/ControlRigPhysics/Private/RigPhysicsBodyExecution.cpp b/Engine/Plugins/Experimental/ControlRigPhysics/Source/ControlRigPhysics/Private/RigPhysicsBodyExecution.cpp
--- a/Engine/Plugins/Experimental/ControlRigPhysics/Source/ControlRigPhysics/Private/RigPhysicsBodyExecution.cpp
+++ b/Engine/Plugins/Experimental/ControlRigPhysics/Source/ControlRigPhysics/Private/RigPhysicsBodyExecution.cpp
@@ -1,17 +1,12 @@
 // Copyright Epic Games, Inc. All Rights Reserved.
 
 #include "RigPhysicsBodyExecution.h"
+#include "ControlRigPhysicsModule.h"
 
 #include "Rigs/RigHierarchyController.h"
 #include "Units/RigUnitContext.h"
 #include "Units/Execution/RigUnit_PrepareForExecution.h"
-#include "ControlRigPhysicsModule.h"
-
 #include "ControlRig.h"
-#if WITH_EDITOR
-#include "RigVMModel/RigVMController.h"
-#include "ControlRigBlueprintLegacy.h"
-#endif
 
 #include UE_INLINE_GENERATED_CPP_BY_NAME(RigPhysicsBodyExecution)
 
diff --git a/Engine/Plugins/Experimental/ControlRigPhysics/Source/ControlRigPhysics/Private/RigPhysicsControlExecution.cpp b/Engine/Plugins/Experimental/ControlRigPhysics/Source/ControlRigPhysics/Private/RigPhysicsControlExecution.cpp
--- a/Engine/Plugins/Experimental/ControlRigPhysics/Source/ControlRigPhysics/Private/RigPhysicsControlExecution.cpp
+++ b/Engine/Plugins/Experimental/ControlRigPhysics/Source/ControlRigPhysics/Private/RigPhysicsControlExecution.cpp
@@ -2,13 +2,11 @@
 
 #include "RigPhysicsControlExecution.h"
 
-#include "RigPhysicsControlComponent.h"
 #include "PhysicsControlHelpers.h"
 
 #include "Rigs/RigHierarchyController.h"
 #include "Units/RigUnitContext.h"
 #include "Units/Execution/RigUnit_PrepareForExecution.h"
-
 #include "ControlRig.h"
 
 #include UE_INLINE_GENERATED_CPP_BY_NAME(RigPhysicsControlExecution)
diff --git a/Engine/Plugins/Experimental/ControlRigPhysics/Source/ControlRigPhysics/Private/RigPhysicsExecution.cpp b/Engine/Plugins/Experimental/ControlRigPhysics/Source/ControlRigPhysics/Private/RigPhysicsExecution.cpp
--- a/Engine/Plugins/Experimental/ControlRigPhysics/Source/ControlRigPhysics/Private/RigPhysicsExecution.cpp
+++ b/Engine/Plugins/Experimental/ControlRigPhysics/Source/ControlRigPhysics/Private/RigPhysicsExecution.cpp
@@ -1,198 +1,23 @@
 // Copyright Epic Games, Inc. All Rights Reserved.
 
 #include "RigPhysicsExecution.h"
-#include "RigPhysicsSimulation.h"
-#include "RigPhysicsBodyComponent.h"
 #include "RigPhysicsJointComponent.h"
-#include "RigPhysicsSolverComponent.h"
 #include "RigPhysicsControlComponent.h"
+#include "RigPhysicsSimulation.h"
+#include "ControlRigPhysicsModule.h"
 
 #include "Rigs/RigHierarchyController.h"
 #include "Units/RigUnitContext.h"
 #include "Units/Execution/RigUnit_PrepareForExecution.h"
-#include "ControlRigPhysicsModule.h"
+#include "ControlRig.h"
 
 #include "PhysicsEngine/PhysicsAsset.h"
 #include "PhysicsEngine/SkeletalBodySetup.h"
 #include "PhysicsEngine/PhysicsConstraintTemplate.h"
 #include "PhysicalMaterials/PhysicalMaterial.h"
 
-#include "ControlRig.h"
-#if WITH_EDITOR
-#include "RigVMModel/RigVMController.h"
-#include "ControlRigBlueprintLegacy.h"
-#endif
-
 #include UE_INLINE_GENERATED_CPP_BY_NAME(RigPhysicsExecution)
 
-TAutoConsoleVariable<int> CVarControlRigPhysicsVisualizationOverride(
-	TEXT("ControlRig.Physics.VisualizationOveride"), -1,
-	TEXT("-1 uses the Step Physics Solver setting, 0 forces visualization to be disabled, 1 forces it to be enabled."));
-
-TAutoConsoleVariable<int> CVarControlRigPhysicsEnableStepSolver(
-	TEXT("ControlRig.Physics.EnableStepSolver"), 1,
-	TEXT("Setting to zero disables stepping the solver"));
-
-
-//======================================================================================================================
-FRigUnit_AddPhysicsSolver_Execute()
-{
-	if (!ExecuteContext.Hierarchy)
-	{
-		return;
-	}
-
-	if (ExecuteContext.GetEventName() != FRigUnit_PrepareForExecution::EventName)
-	{
-		UE_CONTROLRIG_RIGUNIT_REPORT_ERROR(TEXT("AddPhysicsSolver can only be used during Setup"));
-		return;
-	}
-
-	if (URigHierarchyController* Controller = ExecuteContext.Hierarchy->GetController())
-	{
-		FRigHierarchyControllerInstructionBracket InstructionBracket(
-			Controller, ExecuteContext.GetInstructionIndex());
-		PhysicsSolverComponentKey = Controller->AddComponent(
-			FRigPhysicsSolverComponent::StaticStruct(), TEXT("PhysicsSolver"), Owner);
-		if (PhysicsSolverComponentKey.IsValid())
-		{
-			if (FRigPhysicsSolverComponent* Component = Cast<FRigPhysicsSolverComponent>(
-				ExecuteContext.Hierarchy->FindComponent(PhysicsSolverComponentKey)))
-			{
-				Component->SolverSettings = SolverSettings;
-				Component->SimulationSpaceSettings = SimulationSpaceSettings;
-			}
-		}
-	}
-}
-
-//======================================================================================================================
-FRigUnit_TrackInputPose_Execute()
-{
-	DECLARE_SCOPE_HIERARCHICAL_COUNTER_RIGUNIT();
-	if (!ExecuteContext.Hierarchy)
-	{
-		return;
-	}
-	if (FRigPhysicsSolverComponent* SolverComponent = Cast<FRigPhysicsSolverComponent>(
-		ExecuteContext.Hierarchy->FindComponent(PhysicsSolverComponentKey)))
-	{
-		if (bForceNumberOfFrames)
-		{
-			SolverComponent->TrackInputCounter = NumberOfFrames;
-		}
-		else
-		{
-			SolverComponent->TrackInputCounter = FMath::Max(SolverComponent->TrackInputCounter, NumberOfFrames);
-		}
-	}
-}
-
-//======================================================================================================================
-FRigUnit_StepPhysicsSolver_Execute()
-{
-	DECLARE_SCOPE_HIERARCHICAL_COUNTER_RIGUNIT();
-	if (!ExecuteContext.Hierarchy)
-	{
-		return;
-	}
-	if (CVarControlRigPhysicsEnableStepSolver.GetValueOnAnyThread() <= 0)
-	{
-		return;
-	}
-	URigHierarchy& Hierarchy = *ExecuteContext.Hierarchy;
-	if (UControlRig* ControlRig = Cast<UControlRig>(Hierarchy.GetOuter()))
-	{
-		if (FRigPhysicsSolverComponent* SolverComponent = Cast<FRigPhysicsSolverComponent>(
-			Hierarchy.FindComponent(PhysicsSolverComponentKey)))
-		{
-			FRigPhysicsSimulation* PhysicsSimulation = Cast<FRigPhysicsSimulation>(
-				ControlRig->GetPhysicsSimulation(PhysicsSolverComponentKey));
-
-			// Make the solver and instantiate if necessary
-			if (!PhysicsSimulation)
-			{
-				TSharedPtr<FRigPhysicsSimulation> NewPhysicsSimulation = MakeShared<FRigPhysicsSimulation>(
-					FName(ControlRig->GetName()));
-
-				if (ControlRig->RegisterPhysicsSimulation(NewPhysicsSimulation, PhysicsSolverComponentKey))
-				{
-					NewPhysicsSimulation->Instantiate(ExecuteContext, Hierarchy, *SolverComponent);
-				}
-				PhysicsSimulation = NewPhysicsSimulation.Get();
-			}
-
-			if (PhysicsSimulation)
-			{
-				const USceneComponent* OwningSceneComponent = ControlRig->GetOwningSceneComponent();
-				const AActor* OwningActorPtr = OwningSceneComponent ? OwningSceneComponent->GetOwner() : nullptr;
-
-				PhysicsSimulation->StepSimulation(
-					ControlRig->GetWorld(), OwningActorPtr, ExecuteContext, Hierarchy, *SolverComponent, 
-					DeltaTimeOverride, SimulationSpaceDeltaTimeOverride, Alpha);
-				int32 VisualizationOverride = CVarControlRigPhysicsVisualizationOverride.GetValueOnAnyThread();
-				if (VisualizationOverride == 1 || 
-					(VisualizationOverride < 0 && VisualizationSettings.bEnableVisualization))
-				{
-					PhysicsSimulation->Draw(
-						ExecuteContext.GetDrawInterface(), SolverComponent->SolverSettings,
-						VisualizationSettings, ExecuteContext.GetWorld());
-				}
-			}
-		}
-	}
-}
-
-//======================================================================================================================
-FRigUnit_SetPhysicsSolverSimulationSpaceSettings_Execute()
-{
-	DECLARE_SCOPE_HIERARCHICAL_COUNTER_RIGUNIT();
-	if (!ExecuteContext.Hierarchy)
-	{
-		return;
-	}
-	URigHierarchy& Hierarchy = *ExecuteContext.Hierarchy;
-	if (FRigPhysicsSolverComponent* SolverComponent = Cast<FRigPhysicsSolverComponent>(
-		Hierarchy.FindComponent(PhysicsSolverComponentKey)))
-	{
-		SolverComponent->SimulationSpaceSettings = SimulationSpaceSettings;
-	}
-}
-
-//======================================================================================================================
-FRigUnit_GetPhysicsSolverSimulationSpaceSettings_Execute()
-{
-	DECLARE_SCOPE_HIERARCHICAL_COUNTER_RIGUNIT();
-	if (!ExecuteContext.Hierarchy)
-	{
-		return;
-	}
-	URigHierarchy& Hierarchy = *ExecuteContext.Hierarchy;
-	if (const FRigPhysicsSolverComponent* SolverComponent = Cast<FRigPhysicsSolverComponent>(
-		Hierarchy.FindComponent(PhysicsSolverComponentKey)))
-	{
-		SimulationSpaceSettings = SolverComponent->SimulationSpaceSettings;
-	}
-}
-
-//======================================================================================================================
-FRigUnit_SetPhysicsSolverExternalVelocity_Execute()
-{
-	DECLARE_SCOPE_HIERARCHICAL_COUNTER_RIGUNIT();
-	if (!ExecuteContext.Hierarchy)
-	{
-		return;
-	}
-	URigHierarchy& Hierarchy = *ExecuteContext.Hierarchy;
-	if (FRigPhysicsSolverComponent* SolverComponent = Cast<FRigPhysicsSolverComponent>(
-		Hierarchy.FindComponent(PhysicsSolverComponentKey)))
-	{
-		SolverComponent->SimulationSpaceSettings.ExternalLinearVelocity = ExternalLinearVelocity;
-		SolverComponent->SimulationSpaceSettings.ExternalAngularVelocity = ExternalAngularVelocity;
-		SolverComponent->SimulationSpaceSettings.ExternalTurbulenceVelocity = ExternalTurbulenceVelocity;
-	}
-}
-
 //======================================================================================================================
 FRigUnit_InstantiatePhysics_Execute()
 {
@@ -874,46 +699,3 @@ FRigUnit_HierarchyImportCollisionFromPhysicsAsset_Execute()
 		}
 	}
 }
-
-//======================================================================================================================
-FRigUnit_GetPhysicsSolverSpaceData_Execute()
-{
-	DECLARE_SCOPE_HIERARCHICAL_COUNTER_RIGUNIT();
-	if (!ExecuteContext.Hierarchy)
-	{
-		return;
-	}
-	if (UControlRig* ControlRig = Cast<UControlRig>(ExecuteContext.Hierarchy->GetOuter()))
-	{
-		if (FRigPhysicsSolverComponent* SolverComponent = Cast<FRigPhysicsSolverComponent>(
-			ExecuteContext.Hierarchy->FindComponent(PhysicsSolverComponentKey)))
-		{
-			if (FRigPhysicsSimulation* PhysicsSimulation = Cast<FRigPhysicsSimulation>(
-				ControlRig->GetPhysicsSimulation(PhysicsSolverComponentKey)))
-			{
-				const FRigPhysicsSimulation::FSimulationSpaceData& SimulationSpaceData = 
-					PhysicsSimulation->GetSimulationSpaceData();
-
-				LinearVelocity = SimulationSpaceData.LinearVelocity;
-				AngularVelocity = SimulationSpaceData.AngularVelocity;
-				LinearAcceleration = SimulationSpaceData.LinearAcceleration;
-				AngularAcceleration = SimulationSpaceData.AngularAcceleration;
-				Gravity = SimulationSpaceData.Gravity;
-			}
-		}
-	}
-}
-
-//======================================================================================================================
-FRigUnit_HierarchySetPhysicsSolverAllowCCD_Execute()
-{
-	if (!ExecuteContext.Hierarchy)
-	{
-		return;
-	}
-	if (FRigPhysicsSolverComponent* SolverComponent = Cast<FRigPhysicsSolverComponent>(
-		ExecuteContext.Hierarchy->FindComponent(PhysicsSolverComponentKey)))
-	{
-		SolverComponent->SolverSettings.bAllowCCD = bAllowCCD;
-	}
-}
diff --git a/Engine/Plugins/Experimental/ControlRigPhysics/Source/ControlRigPhysics/Private/RigPhysicsJointExecution.cpp b/Engine/Plugins/Experimental/ControlRigPhysics/Source/ControlRigPhysics/Private/RigPhysicsJointExecution.cpp
--- a/Engine/Plugins/Experimental/ControlRigPhysics/Source/ControlRigPhysics/Private/RigPhysicsJointExecution.cpp
+++ b/Engine/Plugins/Experimental/ControlRigPhysics/Source/ControlRigPhysics/Private/RigPhysicsJointExecution.cpp
@@ -1,23 +1,17 @@
 // Copyright Epic Games, Inc. All Rights Reserved.
 
 #include "RigPhysicsJointExecution.h"
-#include "RigPhysicsJointComponent.h"
+#include "ControlRigPhysicsModule.h"
+
+#include "PhysicsControlHelpers.h"
 
 #include "Rigs/RigHierarchyController.h"
 #include "Units/RigUnitContext.h"
 #include "Units/Execution/RigUnit_PrepareForExecution.h"
-#include "ControlRigPhysicsModule.h"
-
-#include "PhysicsControlHelpers.h"
+#include "ControlRig.h"
 
 #include "Chaos/ChaosConstraintSettings.h"
 
-#include "ControlRig.h"
-#if WITH_EDITOR
-#include "RigVMModel/RigVMController.h"
-#include "ControlRigBlueprintLegacy.h"
-#endif
-
 #include UE_INLINE_GENERATED_CPP_BY_NAME(RigPhysicsJointExecution)
 
 //======================================================================================================================
diff --git a/Engine/Plugins/Experimental/ControlRigPhysics/Source/ControlRigPhysics/Private/RigPhysicsSimulation.cpp b/Engine/Plugins/Experimental/ControlRigPhysics/Source/ControlRigPhysics/Private/RigPhysicsSimulation.cpp
--- a/Engine/Plugins/Experimental/ControlRigPhysics/Source/ControlRigPhysics/Private/RigPhysicsSimulation.cpp
+++ b/Engine/Plugins/Experimental/ControlRigPhysics/Source/ControlRigPhysics/Private/RigPhysicsSimulation.cpp
@@ -1,7 +1,7 @@
 // Copyright Epic Games, Inc. All Rights Reserved.
 
 #include "RigPhysicsSimulation.h"
-#include "RigPhysicsSimulationSpace.inl"
+#include "RigPhysicsSimulation_Space.inl"
 
 #include "RigPhysicsBodyComponent.h"
 #include "RigPhysicsJointComponent.h"
diff --git a/Engine/Plugins/Experimental/ControlRigPhysics/Source/ControlRigPhysics/Private/RigPhysicsSimulation.h b/Engine/Plugins/Experimental/ControlRigPhysics/Source/ControlRigPhysics/Private/RigPhysicsSimulation.h
--- a/Engine/Plugins/Experimental/ControlRigPhysics/Source/ControlRigPhysics/Private/RigPhysicsSimulation.h
+++ b/Engine/Plugins/Experimental/ControlRigPhysics/Source/ControlRigPhysics/Private/RigPhysicsSimulation.h
@@ -4,7 +4,6 @@
 
 #include "RigPhysicsData.h"
 
-#include "PhysicsControlData.h"
 #include "PhysicsControlPoseData.h"
 
 #include "Rigs/RigHierarchyComponents.h"
@@ -139,7 +138,6 @@ struct FRigControlRecord
 	bool bCachedHasStrengthOrDamping = true;
 };
 
-
 //======================================================================================================================
 // This represents the low level simulation, plus all the objects and controls we make to go in it
 //======================================================================================================================
diff --git a/Engine/Plugins/Experimental/ControlRigPhysics/Source/ControlRigPhysics/Private/RigPhysicsSimulationSpace.cpp b/Engine/Plugins/Experimental/ControlRigPhysics/Source/ControlRigPhysics/Private/RigPhysicsSimulationSpace.cpp
deleted file mode 100644
diff --git a/Engine/Plugins/Experimental/ControlRigPhysics/Source/ControlRigPhysics/Private/RigPhysicsSimulationSpace.inl b/Engine/Plugins/Experimental/ControlRigPhysics/Source/ControlRigPhysics/Private/RigPhysicsSimulationSpace.inl
deleted file mode 100644
diff --git a/Engine/Plugins/Experimental/ControlRigPhysics/Source/ControlRigPhysics/Private/RigPhysicsSimulationVisualization.cpp b/Engine/Plugins/Experimental/ControlRigPhysics/Source/ControlRigPhysics/Private/RigPhysicsSimulationVisualization.cpp
deleted file mode 100644
diff --git a/Engine/Plugins/Experimental/ControlRigPhysics/Source/ControlRigPhysics/Private/RigPhysicsSimulationWorldObjects.cpp b/Engine/Plugins/Experimental/ControlRigPhysics/Source/ControlRigPhysics/Private/RigPhysicsSimulationWorldObjects.cpp
deleted file mode 100644
diff --git a/Engine/Plugins/Experimental/ControlRigPhysics/Source/ControlRigPhysics/Private/RigPhysicsSimulation_Space.cpp b/Engine/Plugins/Experimental/ControlRigPhysics/Source/ControlRigPhysics/Private/RigPhysicsSimulation_Space.cpp
new file mode 100644
--- /dev/null
+++ b/Engine/Plugins/Experimental/ControlRigPhysics/Source/ControlRigPhysics/Private/RigPhysicsSimulation_Space.cpp
@@ -0,0 +1,322 @@
+// Copyright Epic Games, Inc. All Rights Reserved.
+
+#include "RigPhysicsSimulation.h"
+#include "RigPhysicsSimulation_Space.inl"
+
+#include "RigPhysicsSolverComponent.h"
+
+#include "Rigs/RigHierarchy.h"
+#include "Rigs/RigHierarchyController.h"
+
+#include "PhysicsControlHelpers.h"
+
+#include "Engine/Engine.h"
+
+#include "Physics/ImmediatePhysics/ImmediatePhysicsActorHandle.h"
+#include "Physics/ImmediatePhysics/ImmediatePhysicsJointHandle.h"
+#include "Physics/ImmediatePhysics/ImmediatePhysicsSimulation.h"
+
+#include "Physics/Experimental/PhysScene_Chaos.h"
+#include "Chaos/ParticleHandle.h"
+#include "Chaos/PBDJointConstraints.h"
+#include "Chaos/ImplicitObject.h"
+#include "Chaos/ShapeInstance.h"
+#include "Chaos/Capsule.h"
+#include "Chaos/PBDJointConstraintUtilities.h"
+
+#include "PhysicsEngine/BodyInstance.h"
+#include "PhysicsEngine/ConstraintInstance.h"
+#include "PhysicsEngine/BodySetup.h"
+
+TAutoConsoleVariable<int> CVarControlRigPhysicsShowSimulationSpaceInfo(
+	TEXT("ControlRig.Physics.ShowSimulationSpaceInfo"), 0,
+	TEXT("Shows information associated with the simulation space used in control rig physics."));
+
+//======================================================================================================================
+void FRigPhysicsSimulation::InitSimulationSpace(
+	const FTransform& ComponentTM,
+	const FTransform& BoneRelComponentTM)
+{
+	SimulationSpaceState.ComponentTM = ComponentTM;
+	SimulationSpaceState.BoneRelComponentTM = BoneRelComponentTM;
+}
+
+//======================================================================================================================
+// Danny TODO We need support for double precision Perlin inputs - for now duplicate.
+namespace FPerlinHelpers
+{
+	// random permutation of 256 numbers, repeated 2x
+	static const int32 Permutation[512] = {
+		63, 9, 212, 205, 31, 128, 72, 59, 137, 203, 195, 170, 181, 115, 165, 40, 116, 139, 175, 225, 132, 99, 222, 2, 41, 15, 197, 93, 169, 90, 228, 43, 221, 38, 206, 204, 73, 17, 97, 10, 96, 47, 32, 138, 136, 30, 219,
+		78, 224, 13, 193, 88, 134, 211, 7, 112, 176, 19, 106, 83, 75, 217, 85, 0, 98, 140, 229, 80, 118, 151, 117, 251, 103, 242, 81, 238, 172, 82, 110, 4, 227, 77, 243, 46, 12, 189, 34, 188, 200, 161, 68, 76, 171, 194,
+		57, 48, 247, 233, 51, 105, 5, 23, 42, 50, 216, 45, 239, 148, 249, 84, 70, 125, 108, 241, 62, 66, 64, 240, 173, 185, 250, 49, 6, 37, 26, 21, 244, 60, 223, 255, 16, 145, 27, 109, 58, 102, 142, 253, 120, 149, 160,
+		124, 156, 79, 186, 135, 127, 14, 121, 22, 65, 54, 153, 91, 213, 174, 24, 252, 131, 192, 190, 202, 208, 35, 94, 231, 56, 95, 183, 163, 111, 147, 25, 67, 36, 92, 236, 71, 166, 1, 187, 100, 130, 143, 237, 178, 158,
+		104, 184, 159, 177, 52, 214, 230, 119, 87, 114, 201, 179, 198, 3, 248, 182, 39, 11, 152, 196, 113, 20, 232, 69, 141, 207, 234, 53, 86, 180, 226, 74, 150, 218, 29, 133, 8, 44, 123, 28, 146, 89, 101, 154, 220, 126,
+		155, 122, 210, 168, 254, 162, 129, 33, 18, 209, 61, 191, 199, 157, 245, 55, 164, 167, 215, 246, 144, 107, 235,
+
+		63, 9, 212, 205, 31, 128, 72, 59, 137, 203, 195, 170, 181, 115, 165, 40, 116, 139, 175, 225, 132, 99, 222, 2, 41, 15, 197, 93, 169, 90, 228, 43, 221, 38, 206, 204, 73, 17, 97, 10, 96, 47, 32, 138, 136, 30, 219,
+		78, 224, 13, 193, 88, 134, 211, 7, 112, 176, 19, 106, 83, 75, 217, 85, 0, 98, 140, 229, 80, 118, 151, 117, 251, 103, 242, 81, 238, 172, 82, 110, 4, 227, 77, 243, 46, 12, 189, 34, 188, 200, 161, 68, 76, 171, 194,
+		57, 48, 247, 233, 51, 105, 5, 23, 42, 50, 216, 45, 239, 148, 249, 84, 70, 125, 108, 241, 62, 66, 64, 240, 173, 185, 250, 49, 6, 37, 26, 21, 244, 60, 223, 255, 16, 145, 27, 109, 58, 102, 142, 253, 120, 149, 160,
+		124, 156, 79, 186, 135, 127, 14, 121, 22, 65, 54, 153, 91, 213, 174, 24, 252, 131, 192, 190, 202, 208, 35, 94, 231, 56, 95, 183, 163, 111, 147, 25, 67, 36, 92, 236, 71, 166, 1, 187, 100, 130, 143, 237, 178, 158,
+		104, 184, 159, 177, 52, 214, 230, 119, 87, 114, 201, 179, 198, 3, 248, 182, 39, 11, 152, 196, 113, 20, 232, 69, 141, 207, 234, 53, 86, 180, 226, 74, 150, 218, 29, 133, 8, 44, 123, 28, 146, 89, 101, 154, 220, 126,
+		155, 122, 210, 168, 254, 162, 129, 33, 18, 209, 61, 191, 199, 157, 245, 55, 164, 167, 215, 246, 144, 107, 235
+	};
+
+	// Gradient functions for 1D, 2D and 3D Perlin noise
+
+	FORCEINLINE float Grad1(int32 Hash, double X)
+	{
+		// Slicing Perlin's 3D improved noise would give us only scales of -1, 0 and 1; this looks pretty bad so let's use a different sampling
+		static const double Grad1Scales[16] = { -8 / 8, -7 / 8., -6 / 8., -5 / 8., -4 / 8., -3 / 8., -2 / 8., -1 / 8., 1 / 8., 2 / 8., 3 / 8., 4 / 8., 5 / 8., 6 / 8., 7 / 8., 8 / 8 };
+		return Grad1Scales[Hash & 15] * X;
+	}
+
+
+	// Curve w/ second derivative vanishing at 0 and 1, from Perlin's improved noise paper
+	FORCEINLINE double SmoothCurve(double X)
+	{
+		return X * X * X * (X * (X * 6.0 - 15.0) + 10.0);
+	}
+
+	FORCEINLINE float PerlinNoise1D(double X)
+	{
+		const double Xfl = FMath::FloorToFloat(X);
+		const int64 Xi = (int64)(Xfl) & 255;
+		X -= Xfl;
+		const double Xm1 = X - 1.0;
+
+		const int32 A = Permutation[Xi];
+		const int32 B = Permutation[Xi + 1];
+
+		const double U = SmoothCurve(X);
+
+		// 2.0 factor to ensure (-1, 1) range
+		return 2.0f * FMath::Lerp(Grad1(A, X), Grad1(B, Xm1), U);
+	}
+}
+
+//======================================================================================================================
+// Note - don't use the space conversion functions here as the state won't have been set yet.
+FRigPhysicsSimulation::FSimulationSpaceData FRigPhysicsSimulation::UpdateSimulationSpaceStateAndCalculateData(
+	const FRigVMExecuteContext&         ExecuteContext, 
+	const URigHierarchy&                Hierarchy,
+	const FRigPhysicsSolverComponent&   SolverComponent,
+	const float                         Dt,
+	const double                        AbsoluteTime)
+{
+	const FRigPhysicsSolverSettings& SolverSettings = SolverComponent.SolverSettings;
+	const FRigPhysicsSimulationSpaceSettings& SimulationSpaceSettings = SolverComponent.SimulationSpaceSettings;
+
+	SimulationSpaceState.ComponentTM = ExecuteContext.GetToWorldSpaceTransform();
+
+	if (SolverSettings.SimulationSpace == ERigPhysicsSimulationSpace::SpaceBone &&  SolverSettings.SpaceBone.IsValid())
+	{
+		SimulationSpaceState.BoneRelComponentTM = Hierarchy.GetGlobalTransform(SolverSettings.SpaceBone);
+	}
+
+	// Record the history - but avoid polluting it with zero Dt updates. What that means is - if we
+	// get a zero-Dt update, then just update our current sim space TM, which means the time delta
+	// from the previous state is actually the current Dt (i.e. don't overwrite the current Dt).
+	if (Dt > SMALL_NUMBER)
+	{
+		SimulationSpaceState.PrevDt = SimulationSpaceState.Dt;
+		SimulationSpaceState.Dt = Dt;
+
+		SimulationSpaceState.PrevPrevSimulationSpaceTM = SimulationSpaceState.PrevSimulationSpaceTM;
+		SimulationSpaceState.PrevSimulationSpaceTM = SimulationSpaceState.SimulationSpaceTM;
+	}
+	SimulationSpaceState.SimulationSpaceTM = GetSpaceTransform(
+		SolverSettings.SimulationSpace, SimulationSpaceState.ComponentTM, SimulationSpaceState.BoneRelComponentTM);
+
+	SimulationSpaceData = FSimulationSpaceData();
+	SimulationSpaceData.Gravity = ::ConvertWorldVectorToSimSpaceNoScale(
+		SolverSettings.SimulationSpace, SolverSettings.Gravity, 
+		SimulationSpaceState.ComponentTM, SimulationSpaceState.BoneRelComponentTM);
+
+	if (SolverSettings.SimulationSpace == ERigPhysicsSimulationSpace::World)
+	{
+		// Danny TODO This is probably redundant unless we support runtime switching of the space
+		InitSimulationSpace(SimulationSpaceState.ComponentTM, SimulationSpaceState.BoneRelComponentTM);
+		SimulationSpaceData.LinearVelocity = SimulationSpaceSettings.ExternalLinearVelocity;
+		SimulationSpaceData.AngularVelocity = FMath::DegreesToRadians(SimulationSpaceSettings.ExternalAngularVelocity);
+		return SimulationSpaceData;
+	}
+
+	// If the timestep is zero, then it doesn't actually matter what the velocity is - but make sure
+	// it doesn't corrupt anything.
+	if (SimulationSpaceState.Dt < SMALL_NUMBER)
+	{
+		SimulationSpaceData.LinearVelocity = FVector::ZeroVector;
+		SimulationSpaceData.AngularVelocity = FVector::ZeroVector;
+		SimulationSpaceData.LinearAcceleration = FVector::ZeroVector;
+		SimulationSpaceData.AngularAcceleration = FVector::ZeroVector;
+		return SimulationSpaceData;
+	}
+
+	// We calculate velocities etc in world space first, and then subsequently convert them into
+	// simulation space.
+
+	// Note that the velocity/accel calculations are intended to track the world/simulation behavior
+	// - not necessarily be the most accurate calculations! For example, we could use one-sided
+	// finite difference approximations, but this wouldn't necessarily be correct.
+
+	// World-space component linear velocity and acceleration
+	SimulationSpaceData.LinearVelocity = UE::PhysicsControl::CalculateLinearVelocity(
+		SimulationSpaceState.PrevSimulationSpaceTM.GetTranslation(),
+		SimulationSpaceState.SimulationSpaceTM.GetTranslation(), SimulationSpaceState.Dt);
+	const FVector PrevSpaceLinearVel = 
+		SimulationSpaceState.PrevDt < SMALL_NUMBER 
+		? SimulationSpaceData.LinearVelocity 
+		: UE::PhysicsControl::CalculateLinearVelocity(
+			SimulationSpaceState.PrevPrevSimulationSpaceTM.GetTranslation(),
+			SimulationSpaceState.PrevSimulationSpaceTM.GetTranslation(), SimulationSpaceState.PrevDt);
+	SimulationSpaceData.LinearAcceleration =
+		(SimulationSpaceData.LinearVelocity - PrevSpaceLinearVel) / SimulationSpaceState.Dt;
+
+	// World-space component angular velocity and acceleration
+	SimulationSpaceData.AngularVelocity = UE::PhysicsControl::CalculateAngularVelocity(
+		SimulationSpaceState.PrevSimulationSpaceTM.GetRotation(),
+		SimulationSpaceState.SimulationSpaceTM.GetRotation(), SimulationSpaceState.Dt);
+	const FVector PrevSpaceAngularVel = 
+		SimulationSpaceState.PrevDt < SMALL_NUMBER 
+		? SimulationSpaceData.AngularVelocity
+		: UE::PhysicsControl::CalculateAngularVelocity(
+			SimulationSpaceState.PrevPrevSimulationSpaceTM.GetRotation(),
+			SimulationSpaceState.PrevSimulationSpaceTM.GetRotation(), SimulationSpaceState.PrevDt);
+	SimulationSpaceData.AngularAcceleration =
+		(SimulationSpaceData.AngularVelocity - PrevSpaceAngularVel) / SimulationSpaceState.Dt;
+
+	// Apply Z scale
+	SimulationSpaceData.LinearVelocity.Z *= SimulationSpaceSettings.VelocityScaleZ;
+	SimulationSpaceData.LinearAcceleration.Z *= SimulationSpaceSettings.VelocityScaleZ;
+
+	bool bLinearAccelerationTrigger = SimulationSpaceSettings.LinearAccelerationThresholdForTeleport > 0 &&
+		SimulationSpaceData.LinearAcceleration.SquaredLength() >
+		FMath::Square(SimulationSpaceSettings.LinearAccelerationThresholdForTeleport);
+	bool bAngularAccelerationTrigger = SimulationSpaceSettings.AngularAccelerationThresholdForTeleport > 0 &&
+		SimulationSpaceData.AngularAcceleration.SquaredLength() >
+		FMath::Square(FMath::DegreesToRadians(SimulationSpaceSettings.AngularAccelerationThresholdForTeleport));
+	bool bPositionTrigger = SimulationSpaceSettings.PositionChangeThresholdForTeleport > 0 &&
+		SimulationSpaceData.LinearVelocity.SquaredLength() >
+		FMath::Square(SimulationSpaceSettings.PositionChangeThresholdForTeleport / SimulationSpaceState.Dt);
+	bool bOrientationTrigger = SimulationSpaceSettings.OrientationChangeThresholdForTeleport > 0 &&
+		SimulationSpaceData.AngularVelocity.SquaredLength() >
+		FMath::Square(FMath::DegreesToRadians(
+			SimulationSpaceSettings.OrientationChangeThresholdForTeleport / SimulationSpaceState.Dt));
+
+	// Check and report teleports
+	if (bLinearAccelerationTrigger || bAngularAccelerationTrigger || bPositionTrigger || bOrientationTrigger)
+	{
+		if (bLinearAccelerationTrigger)
+		{
+			UE_LOG(LogRigPhysics, Log, TEXT("Detected linear Acceleration (%f > %f) teleport in %s"),
+				SimulationSpaceData.LinearAcceleration.Length(), 
+				SimulationSpaceSettings.LinearAccelerationThresholdForTeleport,
+				*OwnerName.ToString());
+		}
+		if (bAngularAccelerationTrigger)
+		{
+			UE_LOG(LogRigPhysics, Log, TEXT("Detected angular Acceleration (%f > %f) teleport in %s"),
+				FMath::RadiansToDegrees(SimulationSpaceData.AngularAcceleration.Length()),
+				SimulationSpaceSettings.AngularAccelerationThresholdForTeleport,
+				*OwnerName.ToString());
+		}
+		if (bPositionTrigger)
+		{
+			UE_LOG(LogRigPhysics, Log, TEXT("Detected position (%f > %f) teleport in %s"),
+				SimulationSpaceData.LinearVelocity.Length() * SimulationSpaceState.Dt,
+				SimulationSpaceSettings.PositionChangeThresholdForTeleport,
+				*OwnerName.ToString());
+		}
+		if (bOrientationTrigger)
+		{
+			UE_LOG(LogRigPhysics, Log, TEXT("Detected orientation (%f > %f) teleport in %s"),
+				FMath::RadiansToDegrees(SimulationSpaceData.AngularVelocity.Length() * SimulationSpaceState.Dt),
+				SimulationSpaceSettings.OrientationChangeThresholdForTeleport,
+				*OwnerName.ToString());
+		}
+
+		// Note that a teleport detection shouldn't change the pose, or the current motion. We just
+		// don't want to bring in that unwanted global motion.
+		SimulationSpaceData.LinearVelocity = FVector::ZeroVector;
+		SimulationSpaceData.AngularVelocity = FVector::ZeroVector;
+		SimulationSpaceData.LinearAcceleration = FVector::ZeroVector;
+		SimulationSpaceData.AngularAcceleration = FVector::ZeroVector;
+
+		// This will stop the next step from using bogus values too.
+		SimulationSpaceState.PrevSimulationSpaceTM = SimulationSpaceState.SimulationSpaceTM;
+		SimulationSpaceState.PrevPrevSimulationSpaceTM = SimulationSpaceState.SimulationSpaceTM;
+		SimulationSpaceState.PrevDt = 0;
+		SimulationSpaceState.Dt = 0;
+		// Avoid cached transforms being used in controls by bumping the update counter. 
+		UpdateCounter += 1;
+	}
+	else
+	{
+		if (SimulationSpaceSettings.bClampLinearVelocity)
+		{
+			SimulationSpaceData.LinearVelocity =
+				SimulationSpaceData.LinearVelocity.GetClampedToMaxSize(SimulationSpaceSettings.MaxLinearVelocity);
+		}
+		if (SimulationSpaceSettings.bClampAngularVelocity)
+		{
+			SimulationSpaceData.AngularVelocity =
+				SimulationSpaceData.AngularVelocity.GetClampedToMaxSize(SimulationSpaceSettings.MaxAngularVelocity);
+		}
+		if (SimulationSpaceSettings.bClampLinearAcceleration)
+		{
+			SimulationSpaceData.LinearAcceleration =
+				SimulationSpaceData.LinearAcceleration.GetClampedToMaxSize(SimulationSpaceSettings.MaxLinearAcceleration);
+		}
+		if (SimulationSpaceSettings.bClampAngularAcceleration)
+		{
+			SimulationSpaceData.AngularAcceleration =
+				SimulationSpaceData.AngularAcceleration.GetClampedToMaxSize(SimulationSpaceSettings.MaxAngularAcceleration);
+		}
+	}
+
+	SimulationSpaceData.LinearVelocity += SimulationSpaceSettings.ExternalLinearVelocity;
+	SimulationSpaceData.AngularVelocity += FMath::DegreesToRadians(SimulationSpaceSettings.ExternalAngularVelocity);
+
+	if (!SimulationSpaceSettings.ExternalTurbulenceVelocity.IsNearlyZero())
+	{
+		FVector T(
+			FPerlinHelpers::PerlinNoise1D(AbsoluteTime),
+			FPerlinHelpers::PerlinNoise1D(AbsoluteTime + 10.0),
+			FPerlinHelpers::PerlinNoise1D(AbsoluteTime + 20.0));
+
+		FVector Turbulence = T * SimulationSpaceSettings.ExternalTurbulenceVelocity;
+
+		SimulationSpaceData.LinearVelocity += Turbulence;
+
+		if (GEngine && CVarControlRigPhysicsShowSimulationSpaceInfo.GetValueOnAnyThread())
+		{
+			GEngine->AddOnScreenDebugMessage(INDEX_NONE, 0.f, FColor::Yellow,
+				FString::Printf(TEXT("Sim %s Turbulence %s"), *OwnerName.ToString(), *Turbulence.ToString()));
+		}
+	}
+
+	// Transform world-space motion into simulation space Danny TODO note that this matches the code
+	// in RBAN, and is doing what the interface requires (i.e. movement of the space in the space of
+	// the space!). Danny TODO note that we don't currently support scaling
+	SimulationSpaceData.LinearVelocity =
+		SimulationSpaceState.SimulationSpaceTM.InverseTransformVectorNoScale(SimulationSpaceData.LinearVelocity);
+	SimulationSpaceData.AngularVelocity =
+		SimulationSpaceState.SimulationSpaceTM.InverseTransformVectorNoScale(SimulationSpaceData.AngularVelocity);
+	SimulationSpaceData.LinearAcceleration =
+		SimulationSpaceState.SimulationSpaceTM.InverseTransformVectorNoScale(SimulationSpaceData.LinearAcceleration);
+	SimulationSpaceData.AngularAcceleration =
+		SimulationSpaceState.SimulationSpaceTM.InverseTransformVectorNoScale(SimulationSpaceData.AngularAcceleration);
+
+	if (GEngine && CVarControlRigPhysicsShowSimulationSpaceInfo.GetValueOnAnyThread())
+	{
+		GEngine->AddOnScreenDebugMessage(INDEX_NONE, 0.f, FColor::Yellow,
+			FString::Printf(TEXT("Sim %s dt %6.2fms Space P %35s V %35s A %35s"),
+				*OwnerName.ToString(), Dt * 1000.0f,
+				*SimulationSpaceState.SimulationSpaceTM.GetLocation().ToString(),
+				*SimulationSpaceData.LinearVelocity.ToString(),
+				*SimulationSpaceData.LinearAcceleration.ToString()));
+	}
+
+	return SimulationSpaceData;
+}
diff --git a/Engine/Plugins/Experimental/ControlRigPhysics/Source/ControlRigPhysics/Private/RigPhysicsSimulation_Space.inl b/Engine/Plugins/Experimental/ControlRigPhysics/Source/ControlRigPhysics/Private/RigPhysicsSimulation_Space.inl
new file mode 100644
--- /dev/null
+++ b/Engine/Plugins/Experimental/ControlRigPhysics/Source/ControlRigPhysics/Private/RigPhysicsSimulation_Space.inl
@@ -0,0 +1,186 @@
+// Copyright Epic Games, Inc. All Rights Reserved.
+
+#pragma once
+
+#include "RigPhysicsSimulation.h"
+#include "PhysicsControlHelpers.h"
+
+#include "Engine/Engine.h"
+
+//======================================================================================================================
+FORCEINLINE FTransform FRigPhysicsSimulation::GetSpaceTransform(
+	ERigPhysicsSimulationSpace Space, const FTransform& ComponentTM, const FTransform& BoneTM)
+{
+	switch (Space)
+	{
+	case ERigPhysicsSimulationSpace::Component: return ComponentTM;
+	case ERigPhysicsSimulationSpace::World: return FTransform::Identity;
+	case ERigPhysicsSimulationSpace::SpaceBone: return BoneTM * ComponentTM;
+	default: ensureMsgf(false, TEXT("Unsupported Simulation Space")); return FTransform::Identity;
+	}
+}
+
+//======================================================================================================================
+FORCEINLINE FTransform FRigPhysicsSimulation::GetSimulationSpaceTransform(const FRigPhysicsSolverSettings& SolverSettings) const
+{
+	switch (SolverSettings.SimulationSpace)
+	{
+	case ERigPhysicsSimulationSpace::Component: return SimulationSpaceState.ComponentTM;
+	case ERigPhysicsSimulationSpace::World: return FTransform::Identity;
+	case ERigPhysicsSimulationSpace::SpaceBone: return 
+		SimulationSpaceState.BoneRelComponentTM * SimulationSpaceState.ComponentTM;
+	default: ensureMsgf(false, TEXT("Unsupported Simulation Space")); return FTransform::Identity;
+	}
+}
+
+//======================================================================================================================
+FORCEINLINE FTransform FRigPhysicsSimulation::ConvertComponentSpaceTransformToSimSpace(
+	const FRigPhysicsSolverSettings& SolverSettings, const FTransform& TM) const
+{
+	switch (SolverSettings.SimulationSpace)
+	{
+	case ERigPhysicsSimulationSpace::Component: return TM;
+	case ERigPhysicsSimulationSpace::World: return TM * SimulationSpaceState.ComponentTM;
+	case ERigPhysicsSimulationSpace::SpaceBone: return TM.GetRelativeTransform(SimulationSpaceState.BoneRelComponentTM);
+	default: ensureMsgf(false, TEXT("Unsupported Simulation Space")); return TM;
+	}
+}
+
+//======================================================================================================================
+FORCEINLINE FVector FRigPhysicsSimulation::ConvertComponentSpacePositionToSimSpace(
+	const FRigPhysicsSolverSettings& SolverSettings, const FVector& P) const
+{
+	switch (SolverSettings.SimulationSpace)
+	{
+	case ERigPhysicsSimulationSpace::Component: return P;
+	case ERigPhysicsSimulationSpace::World:  return SimulationSpaceState.ComponentTM.TransformPositionNoScale(P);
+	case ERigPhysicsSimulationSpace::SpaceBone: return SimulationSpaceState.BoneRelComponentTM.InverseTransformPositionNoScale(P);
+	default: ensureMsgf(false, TEXT("Unsupported Simulation Space")); return P;
+	}
+}
+
+//======================================================================================================================
+FORCEINLINE FVector FRigPhysicsSimulation::ConvertComponentSpaceVectorToSimSpace(
+	const FRigPhysicsSolverSettings& SolverSettings, const FVector& V) const
+{
+	switch (SolverSettings.SimulationSpace)
+	{
+	case ERigPhysicsSimulationSpace::Component: return V;
+	case ERigPhysicsSimulationSpace::World: return SimulationSpaceState.ComponentTM.TransformVectorNoScale(V);
+	case ERigPhysicsSimulationSpace::SpaceBone: 
+		return SimulationSpaceState.BoneRelComponentTM.InverseTransformVectorNoScale(V);
+	default: ensureMsgf(false, TEXT("Unsupported Simulation Space")); return V;
+	}
+}
+
+//======================================================================================================================
+FORCEINLINE FTransform FRigPhysicsSimulation::ConvertSimSpaceTransformToComponentSpace(
+	const FRigPhysicsSolverSettings& SolverSettings, const FTransform& TM) const
+{
+	switch (SolverSettings.SimulationSpace)
+	{
+	case ERigPhysicsSimulationSpace::Component: return TM;
+	case ERigPhysicsSimulationSpace::World: return TM.GetRelativeTransform(SimulationSpaceState.ComponentTM);
+	case ERigPhysicsSimulationSpace::SpaceBone: return TM * SimulationSpaceState.BoneRelComponentTM;
+	default: ensureMsgf(false, TEXT("Unsupported Simulation Space")); return TM;
+	}
+}
+
+//======================================================================================================================
+FORCEINLINE FVector FRigPhysicsSimulation::ConvertSimSpacePositionToComponentSpace(
+	const FRigPhysicsSolverSettings& SolverSettings, const FVector& Position) const
+{
+	switch (SolverSettings.SimulationSpace)
+	{
+	case ERigPhysicsSimulationSpace::Component: return Position;
+	case ERigPhysicsSimulationSpace::World: return SimulationSpaceState.ComponentTM.InverseTransformPositionNoScale(Position);
+	case ERigPhysicsSimulationSpace::SpaceBone: return SimulationSpaceState.BoneRelComponentTM.TransformPositionNoScale(Position);
+	default: ensureMsgf(false, TEXT("Unsupported Simulation Space")); return Position;
+	}
+}
+
+//======================================================================================================================
+FORCEINLINE FVector FRigPhysicsSimulation::ConvertSimSpaceVectorToComponentSpace(
+	const FRigPhysicsSolverSettings& SolverSettings, const FVector& Vector) const
+{
+	switch (SolverSettings.SimulationSpace)
+	{
+	case ERigPhysicsSimulationSpace::Component: return Vector;
+	case ERigPhysicsSimulationSpace::World: return SimulationSpaceState.ComponentTM.InverseTransformVectorNoScale(Vector);
+	case ERigPhysicsSimulationSpace::SpaceBone: return SimulationSpaceState.BoneRelComponentTM.TransformVectorNoScale(Vector);
+	default: ensureMsgf(false, TEXT("Unsupported Simulation Space")); return Vector;
+	}
+}
+
+//======================================================================================================================
+FORCEINLINE FVector ConvertWorldVectorToSimSpaceNoScale(
+	ERigPhysicsSimulationSpace Space, const FVector& WorldVector, 
+	const FTransform& ComponentTM, const FTransform& BoneTM)
+{
+	switch (Space)
+	{
+	case ERigPhysicsSimulationSpace::Component: return ComponentTM.InverseTransformVectorNoScale(WorldVector);
+	case ERigPhysicsSimulationSpace::World: return WorldVector;
+	case ERigPhysicsSimulationSpace::SpaceBone:
+		return BoneTM.InverseTransformVectorNoScale(ComponentTM.InverseTransformVectorNoScale(WorldVector));
+	default: ensureMsgf(false, TEXT("Unsupported Simulation Space")); return WorldVector;
+	}
+}
+
+//======================================================================================================================
+FORCEINLINE FVector FRigPhysicsSimulation::ConvertWorldVectorToSimSpaceNoScale(
+	const FRigPhysicsSolverSettings& SolverSettings, const FVector& WorldVector) const
+{
+	switch (SolverSettings.SimulationSpace)
+	{
+	case ERigPhysicsSimulationSpace::Component: 
+		return SimulationSpaceState.ComponentTM.InverseTransformVectorNoScale(WorldVector);
+	case ERigPhysicsSimulationSpace::World: return WorldVector;
+	case ERigPhysicsSimulationSpace::SpaceBone:
+		return SimulationSpaceState.BoneRelComponentTM.InverseTransformVectorNoScale(
+			SimulationSpaceState.ComponentTM.InverseTransformVectorNoScale(WorldVector));
+	default: ensureMsgf(false, TEXT("Unsupported Simulation Space")); return WorldVector;
+	}
+}
+
+//======================================================================================================================
+FORCEINLINE FVector FRigPhysicsSimulation::ConvertWorldPositionToSimSpaceNoScale(
+	const FRigPhysicsSolverSettings& SolverSettings, const FVector& WorldPosition) const
+{
+	switch (SolverSettings.SimulationSpace)
+	{
+	case ERigPhysicsSimulationSpace::Component: 
+		return SimulationSpaceState.ComponentTM.InverseTransformPositionNoScale(WorldPosition);
+	case ERigPhysicsSimulationSpace::World: return WorldPosition;
+	case ERigPhysicsSimulationSpace::SpaceBone:
+		return SimulationSpaceState.BoneRelComponentTM.InverseTransformPositionNoScale(
+			SimulationSpaceState.ComponentTM.InverseTransformPositionNoScale(WorldPosition));
+	default: ensureMsgf(false, TEXT("Unsupported Simulation Space")); return WorldPosition;
+	}
+}
+
+//======================================================================================================================
+FORCEINLINE FTransform FRigPhysicsSimulation::ConvertWorldTransformToSimSpace(
+	const FRigPhysicsSolverSettings& SolverSettings, const FTransform& WorldTM) const
+{
+	switch (SolverSettings.SimulationSpace)
+	{
+	case ERigPhysicsSimulationSpace::Component: return WorldTM.GetRelativeTransform(SimulationSpaceState.ComponentTM);
+	case ERigPhysicsSimulationSpace::World: return WorldTM;
+	case ERigPhysicsSimulationSpace::SpaceBone:
+		return WorldTM.GetRelativeTransform(SimulationSpaceState.BoneRelComponentTM * SimulationSpaceState.ComponentTM);
+	default: ensureMsgf(false, TEXT("Unsupported Simulation Space")); return WorldTM;
+	}
+}
+
+//======================================================================================================================
+FORCEINLINE FTransform FRigPhysicsSimulation::ConvertCollisionSpaceTransformToSimSpace(
+	const FRigPhysicsSolverSettings& SolverSettings, const FTransform& TM) const
+{
+	FTransform SimSpaceTM = GetSpaceTransform(
+		SolverSettings.SimulationSpace, SimulationSpaceState.ComponentTM, SimulationSpaceState.BoneRelComponentTM);
+	FTransform CollisionSpaceTM = GetSpaceTransform(
+		SolverSettings.CollisionSpace, SimulationSpaceState.ComponentTM, SimulationSpaceState.BoneRelComponentTM);
+	FTransform WorldSpaceTM = TM * CollisionSpaceTM;
+	return WorldSpaceTM.GetRelativeTransform(SimSpaceTM);
+}
diff --git a/Engine/Plugins/Experimental/ControlRigPhysics/Source/ControlRigPhysics/Private/RigPhysicsSimulation_Visualization.cpp b/Engine/Plugins/Experimental/ControlRigPhysics/Source/ControlRigPhysics/Private/RigPhysicsSimulation_Visualization.cpp
new file mode 100644
--- /dev/null
+++ b/Engine/Plugins/Experimental/ControlRigPhysics/Source/ControlRigPhysics/Private/RigPhysicsSimulation_Visualization.cpp
@@ -0,0 +1,529 @@
+// Copyright Epic Games, Inc. All Rights Reserved.
+
+#include "RigPhysicsSimulation.h"
+#include "RigPhysicsSimulation_Space.inl"
+
+#include "Rigs/RigHierarchy.h"
+#include "Rigs/RigHierarchyController.h"
+
+#include "RigVMCore/RigVMDrawInterface.h"
+
+#include "Physics/ImmediatePhysics/ImmediatePhysicsActorHandle.h"
+#include "Physics/ImmediatePhysics/ImmediatePhysicsJointHandle.h"
+#include "Physics/ImmediatePhysics/ImmediatePhysicsSimulation.h"
+
+#include "Physics/Experimental/PhysScene_Chaos.h"
+#include "Chaos/ParticleHandle.h"
+#include "Chaos/PBDJointConstraints.h"
+#include "Chaos/ImplicitObject.h"
+#include "Chaos/ShapeInstance.h"
+#include "Chaos/Capsule.h"
+#include "Chaos/PBDJointConstraintUtilities.h"
+
+#include "Engine/Engine.h"
+#include "Materials/MaterialInstanceDynamic.h"
+
+#include "Math/Color.h"
+
+// Danny TODO expose these in the visualization settings
+namespace RigPhysicsSolverDraw
+{
+	FLinearColor DynamicCollisionColor = FColor::Yellow;
+	FLinearColor DynamicNoCollisionColor = FColor::Orange;
+	FLinearColor KinematicCollisionColor = FColor::Blue;
+	FLinearColor KinematicNoCollisionColor = FColor::Purple;
+
+	FLinearColor ActiveContactColor = FColor::Red;
+	FLinearColor InactiveContactColor = FColor::Silver;
+}
+
+TAutoConsoleVariable<int> CVarControlRigPhysicsShowActiveContactsOverride(
+	TEXT("ControlRig.Physics.ShowActiveContactsOveride"), -1,
+	TEXT("Whether to draw active contacts (requires visualization to be enabled). -1 uses the visualization setting, 0 forces drawing to be disabled, 1 forces it to be enabled."));
+
+TAutoConsoleVariable<int> CVarControlRigPhysicsShowInactiveContactsOverride(
+	TEXT("ControlRig.Physics.ShowInactiveContactsOveride"), -1,
+	TEXT("Whether to draw inactive contacts (requires visualization to be enabled). -1 uses the visualization setting, 0 forces drawing to be disabled, 1 forces it to be enabled."));
+
+//======================================================================================================================
+static FLinearColor GetColor(bool bIsKinematic, bool bCollides)
+{
+	if (bIsKinematic)
+	{
+		return bCollides 
+			? RigPhysicsSolverDraw::KinematicCollisionColor : RigPhysicsSolverDraw::KinematicNoCollisionColor;
+	}
+	else
+	{
+		return bCollides 
+			? RigPhysicsSolverDraw::DynamicCollisionColor : RigPhysicsSolverDraw::DynamicNoCollisionColor;
+	}
+}
+
+//======================================================================================================================
+static void DrawShapes(
+	FRigVMDrawInterface&                    DI,
+	const FRigPhysicsVisualizationSettings& VisualizationSettings,
+	const bool                              bIsKinematic,
+	const bool                              bCollides,
+	const Chaos::FGeometryParticleHandle*   Particle,
+	const FTransform&                       ShapeTransform, 
+	const Chaos::FImplicitObject*           ImplicitObject,
+	const Chaos::FPerShapeData*             Shape)
+{
+	using namespace Chaos;
+
+	if (!ImplicitObject)
+	{
+		return;
+	}
+	const EImplicitObjectType PackedType = ImplicitObject->GetType(); // Type includes scaling and instancing data
+	const EImplicitObjectType InnerType = GetInnerType(ImplicitObject->GetType());
+
+	// For simplicity, we're going to assume no scaling, instancing, transforming etc
+
+	if (!ImplicitObject)
+	{
+		return;
+	}
+
+	switch (InnerType)
+	{
+	case ImplicitObjectType::Transformed:
+	{
+		const TImplicitObjectTransformed<FReal, 3>* Transformed = 
+			ImplicitObject->GetObject<TImplicitObjectTransformed<FReal, 3>>();
+		FTransform TransformedTransform = FTransform(
+			ShapeTransform.GetRotation() * Transformed->GetTransform().GetRotation(),
+			ShapeTransform.TransformPosition(Transformed->GetTransform().GetLocation()));
+		DrawShapes(
+			DI, VisualizationSettings, bIsKinematic, bCollides, Particle, 
+			TransformedTransform, Transformed->GetTransformedObject(), Shape);
+		return;
+	}
+	case ImplicitObjectType::Sphere:
+	{
+		const Chaos::FSphere* Sphere = ImplicitObject->GetObject<Chaos::FSphere>();
+		DI.DrawSphere(ShapeTransform, FTransform(FVec3(Sphere->GetCenterOfMass())), Sphere->GetRadiusf(),
+			GetColor(bIsKinematic, bCollides),
+			VisualizationSettings.LineThickness, VisualizationSettings.ShapeDetail);
+		break;
+	}
+	case ImplicitObjectType::Box:
+	{
+		const TBox<FReal, 3>* Box = ImplicitObject->GetObject<TBox<FReal, 3>>(); // Axis aligned
+		DI.DrawBox(
+			FTransform(FQuat::Identity, Box->GetCenter()) * ShapeTransform, 
+			FTransform(FQuat::Identity, FVector::ZeroVector, Box->Extents()),
+			GetColor(bIsKinematic, bCollides),
+			VisualizationSettings.LineThickness);
+		break;
+	}
+	case ImplicitObjectType::Capsule:
+	{
+		const Chaos::FCapsule* Capsule = ImplicitObject->GetObject<Chaos::FCapsule>();
+		const FQuat Q = FRotationMatrix::MakeFromZ(Capsule->GetAxis()).ToQuat();
+		DI.DrawCapsule(ShapeTransform, FTransform(Q, Capsule->GetCenterOfMass()),
+			Capsule->GetRadiusf(), Capsule->GetHeightf(),
+			GetColor(bIsKinematic, bCollides),
+			VisualizationSettings.LineThickness);
+		break;
+	}
+	default:
+		// If we don't know what it is, don't draw it.
+		break;
+	}
+}
+
+//======================================================================================================================
+static void DrawActor(
+	FRigVMDrawInterface&                    DI,
+	const FRigPhysicsVisualizationSettings& VisualizationSettings,
+	const FTransform&                       SpaceTransform,
+	const ImmediatePhysics::FActorHandle&   ActorHandle)
+{
+	bool bIsKinematic = ActorHandle.GetIsKinematic();
+	bool bCollides = ActorHandle.GetHasCollision();
+
+	const Chaos::FGeometryParticleHandle* GeometryParticleHandle = ActorHandle.GetParticle();
+	const FTransform ParticleTransform = GeometryParticleHandle->GetTransformXR() * SpaceTransform;
+	for (const Chaos::FShapeInstancePtr& ShapeInstance : GeometryParticleHandle->ShapeInstances())
+	{
+		const Chaos::FImplicitObject* ImplicitObject = ShapeInstance->GetGeometry();
+		DrawShapes(
+			DI, VisualizationSettings, bIsKinematic, bCollides, GeometryParticleHandle,
+			ParticleTransform, ImplicitObject, ShapeInstance.Get());
+	}
+
+	if (VisualizationSettings.bShowCentreOfMass)
+	{
+		float Size = 1.0f * VisualizationSettings.ShapeSize;
+		DI.DrawAxes(
+			ParticleTransform, ActorHandle.GetLocalCoMTransform(),
+			Size, VisualizationSettings.LineThickness);
+	}
+}
+
+//======================================================================================================================
+enum
+{
+	TwistIndex = 0,
+	Swing1Index = 1,
+	Swing2Index = 2
+};
+
+//======================================================================================================================
+// Returns the limit range - it will be the value passed in unless it the motion is free or locked. 
+float GetLimitAngleRadians(float LimitAngle, Chaos::EJointMotionType LimitType)
+{
+	switch (LimitType)
+	{
+	case Chaos::EJointMotionType::Free: return UE_PI;
+	case Chaos::EJointMotionType::Locked: return 0.f;
+	default: return LimitAngle;
+	}
+}
+
+//======================================================================================================================
+static void DrawJoint(
+	FRigVMDrawInterface&                    DI,
+	const FRigPhysicsVisualizationSettings& VisualizationSettings,
+	const FTransform&                       SpaceTransform,
+	const ImmediatePhysics::FJointHandle&   JointHandle,
+	const bool                              bDrawAsAxes)
+{
+	const Chaos::FPBDJointConstraintHandle* ConstraintHandle = JointHandle.GetConstraint();
+	if (!ConstraintHandle)
+	{
+		return;
+	}
+
+	const Chaos::TVec2<const ImmediatePhysics_Chaos::FActorHandle*>& ActorHandles = JointHandle.GetActorHandles();
+	const ImmediatePhysics_Chaos::FActorHandle* ChildActor = ActorHandles[0];
+	const ImmediatePhysics_Chaos::FActorHandle* ParentActor = ActorHandles[1];
+
+	FTransform ChildActorTM = (ChildActor ? ChildActor->GetWorldTransform() : FTransform()) * SpaceTransform;
+	FTransform ParentActorTM = (ParentActor ? ParentActor->GetWorldTransform() : FTransform()) * SpaceTransform;
+
+	bool bChildIsKinematic = ChildActor ? ChildActor->GetIsKinematic() : true;
+	bool bParentIsKinematic = ParentActor ? ParentActor->GetIsKinematic() : true;
+
+	float Size = 5.0f * VisualizationSettings.ShapeSize;
+
+	const Chaos::FPBDJointSettings& JointSettings = ConstraintHandle->GetSettings();
+
+	FTransform DialFrame = JointSettings.ConnectorTransforms[0] * ChildActorTM;
+	FTransform LimitFrame = JointSettings.ConnectorTransforms[1] * ParentActorTM;
+
+	if (bDrawAsAxes)
+	{
+		DI.DrawAxes(
+			FTransform(JointSettings.AngularDrivePositionTarget, JointSettings.LinearDrivePositionTarget) * LimitFrame, 
+			FTransform(), Size, VisualizationSettings.LineThickness * 2);
+		DI.DrawAxes(DialFrame, FTransform(), Size, VisualizationSettings.LineThickness);
+	}
+
+	{
+		// See FConstraintInstance::DrawConstraintImp for inspiration
+
+		// There seems to be a swap between swing1 and swing2 compared to GetSwingTwistAngles
+		FVector LimitAngleRadians(
+			GetLimitAngleRadians(JointSettings.AngularLimits[0], JointSettings.AngularMotionTypes[0]),
+			GetLimitAngleRadians(JointSettings.AngularLimits[2], JointSettings.AngularMotionTypes[2]),
+			GetLimitAngleRadians(JointSettings.AngularLimits[1], JointSettings.AngularMotionTypes[1])
+		);
+
+		Chaos::FReal TwistAngle, Swing1Angle, Swing2Angle;
+		const FQuat LimitQ = LimitFrame.GetRotation();
+		FQuat DialQ = DialFrame.GetRotation();
+		DialQ.EnforceShortestArcWith(LimitQ);
+		Chaos::FPBDJointUtilities::GetSwingTwistAngles(LimitQ, DialQ, TwistAngle, Swing1Angle, Swing2Angle);
+
+		bool bDrawViolatedLimits = true;
+
+		bool bTwistViolated = false;
+		bool bSwing1Violated = false;
+		bool bSwing2Violated = false;
+
+		if (bDrawViolatedLimits)
+		{
+			bTwistViolated = JointSettings.AngularMotionTypes[TwistIndex] == Chaos::EJointMotionType::Limited &&
+				FMath::Abs(TwistAngle) > LimitAngleRadians[0];
+
+			bSwing1Violated = JointSettings.AngularMotionTypes[Swing1Index] == Chaos::EJointMotionType::Limited &&
+				FMath::Abs(Swing1Angle) > LimitAngleRadians[1];
+
+			bSwing2Violated = JointSettings.AngularMotionTypes[Swing2Index] == Chaos::EJointMotionType::Limited &&
+				FMath::Abs(Swing2Angle) > LimitAngleRadians[2];
+		}
+
+		const bool bLockSwing1 = JointSettings.AngularMotionTypes[Swing1Index] == Chaos::EJointMotionType::Locked;
+		const bool bLockSwing2 = JointSettings.AngularMotionTypes[Swing2Index] == Chaos::EJointMotionType::Locked;
+		const bool bLockAllSwing = bLockSwing1 && bLockSwing2;
+
+		check(GEngine->ConstraintLimitMaterialX && GEngine->ConstraintLimitMaterialY && GEngine->ConstraintLimitMaterialZ);
+		static UMaterialInterface* LimitMaterialX = GEngine->ConstraintLimitMaterialX;
+		static UMaterialInterface* LimitMaterialXAxis = GEngine->ConstraintLimitMaterialXAxis;
+		static UMaterialInterface* LimitMaterialY = GEngine->ConstraintLimitMaterialY;
+		static UMaterialInterface* LimitMaterialYAxis = GEngine->ConstraintLimitMaterialYAxis;
+		static UMaterialInterface* LimitMaterialZ = GEngine->ConstraintLimitMaterialZ;
+		static UMaterialInterface* LimitMaterialZAxis = GEngine->ConstraintLimitMaterialZAxis;
+
+		float ArrowLength = Size * 1.05f; // stick out a little bit 
+
+		// If swing is limited (but not locked) - draw the swing limit cone.
+		if (!bLockAllSwing)
+		{
+			if (JointSettings.AngularMotionTypes[Swing1Index] == Chaos::EJointMotionType::Free && 
+				JointSettings.AngularMotionTypes[Swing2Index] == Chaos::EJointMotionType::Free)
+			{
+				DI.DrawSphere(LimitFrame, FTransform(), Size * 0.2f, FColor::White, 
+					VisualizationSettings.LineThickness, VisualizationSettings.ShapeDetail);
+			}
+			else
+			{
+				FTransform ConeLimitTM = LimitFrame;
+				ConeLimitTM.SetScale3D(FVector(Size));
+				const float Swing1LimitAngle = LimitAngleRadians[Swing1Index];
+				const float Swing2LimitAngle = LimitAngleRadians[Swing2Index];
+				DI.DrawCone(
+					ConeLimitTM, FTransform(), Swing1LimitAngle, Swing2LimitAngle, VisualizationSettings.ShapeDetail, 
+					true, FColor::Red, LimitMaterialX->GetRenderProxy(), VisualizationSettings.LineThickness);
+			}
+			// Draw the swing Dial indicator - shows the current orientation of the child frame
+			// relative to the parent frame on the swing axis.
+			// Start the arrow at the limit position as it can be confusing if there is joint separation
+			DI.DrawArrow(FTransform(LimitFrame.GetTranslation()), 
+				DialFrame.GetUnitAxis(EAxis::X) * ArrowLength,
+				DialFrame.GetUnitAxis(EAxis::Y), FColor::Red, VisualizationSettings.LineThickness);
+
+			if (bSwing1Violated || bSwing2Violated)
+			{
+				// Drawing a sphere seems a bit heavy... but it appears that if DrawPoint is used then they get
+				// culled when you get close!
+				DI.DrawSphere(FTransform(
+					DialFrame.GetRotation(),
+					LimitFrame.GetTranslation() + DialFrame.GetUnitAxis(EAxis::X) * ArrowLength),
+					FTransform(),
+					ArrowLength * 0.01f, FColor::Orange, VisualizationSettings.LineThickness * 4);
+			}
+
+		}
+
+		// Draw the twist limit
+		if (JointSettings.AngularMotionTypes[TwistIndex] != Chaos::EJointMotionType::Locked)
+		{
+			// Draw as a flat cone
+			FTransform ConeLimitTM = LimitFrame;
+			ConeLimitTM.SetScale3D(FVector(Size));
+			const float TwistLimitAngle = LimitAngleRadians[TwistIndex];
+			DI.DrawCone(ConeLimitTM, FTransform(
+				FQuat::MakeFromRotationVector(FVector(0, 1, 0) * -HALF_PI) * 
+				FQuat::MakeFromRotationVector(FVector(0, 0, 1) * HALF_PI)),
+				TwistLimitAngle, 0.0f, VisualizationSettings.ShapeDetail, true, FColor::Green,
+				LimitMaterialYAxis->GetRenderProxy(), VisualizationSettings.LineThickness);
+
+			// Draw the twist Dial indicator - shows the current orientation of the child frame
+			// relative to the parent frame on the twist axis.
+			FQuat Rot(LimitFrame.GetUnitAxis(EAxis::X), TwistAngle);
+			FVector TwistArrow = Rot * LimitFrame.GetUnitAxis(EAxis::Y);
+			FVector TwistSideDir = FVector::CrossProduct(LimitFrame.GetUnitAxis(EAxis::Z), TwistArrow);
+
+			DI.DrawArrow(FTransform(LimitFrame.GetTranslation()),
+				TwistArrow * ArrowLength,
+				TwistSideDir, FColor::Green, VisualizationSettings.LineThickness);
+
+			if (bTwistViolated)
+			{
+				DI.DrawSphere(FTransform(
+					DialFrame.GetRotation(),
+					LimitFrame.GetTranslation() + TwistArrow * ArrowLength),
+					FTransform(),
+					ArrowLength * 0.01f, FColor::Orange, VisualizationSettings.LineThickness * 4);
+			}
+		}
+	}
+}
+
+//======================================================================================================================
+static bool IsConstraintDriveActive(const Chaos::FPBDJointConstraintHandle& Constraint)
+{
+	const Chaos::FPBDJointSettings& Settings = Constraint.GetSettings();
+	if (Settings.bAngularSLerpPositionDriveEnabled || Settings.bAngularSwingPositionDriveEnabled || Settings.bAngularTwistPositionDriveEnabled)
+	{
+		if (Settings.AngularDriveStiffness.SizeSquared() > 0.0)
+		{
+			return true;
+		}
+	}
+	if (Settings.bAngularSLerpVelocityDriveEnabled || Settings.bAngularSwingVelocityDriveEnabled || Settings.bAngularTwistVelocityDriveEnabled)
+	{
+		if (Settings.AngularDriveDamping.SizeSquared() > 0.0)
+		{
+			return true;
+		}
+	}
+	if (Settings.bLinearPositionDriveEnabled.Max() && Settings.LinearDriveStiffness.SizeSquared() > 0.0)
+	{
+		return true;
+	}
+	if (Settings.bLinearVelocityDriveEnabled.Max() && Settings.LinearDriveDamping.SizeSquared() > 0.0)
+	{
+		return true;
+	}
+	return false;
+}
+
+//======================================================================================================================
+void FRigPhysicsSimulation::Draw(
+	FRigVMDrawInterface*                    DI,
+	const FRigPhysicsSolverSettings&        SolverSettings,
+	const FRigPhysicsVisualizationSettings& VisualizationSettings,
+	const UWorld*                           DebugWorld) const
+{
+	DECLARE_SCOPE_HIERARCHICAL_COUNTER_FUNC()
+	QUICK_SCOPE_CYCLE_COUNTER(STAT_RigPhysics_Draw);
+
+	// Lock against access to the simulation, in case of access by the WorldObject task
+	UE::TScopeLock<FTransactionallySafeCriticalSection> Lock(*SimulationMutex.Get());
+
+#if CHAOS_DEBUG_DRAW
+	// Danny TODO enable/disable chaos debug draw with a cvar or flag
+	if (DebugWorld && DebugWorld->GetPhysicsScene() && DebugWorld->GetPhysicsScene()->GetDebugDrawScene())
+	{
+		Simulation->SetDebugDrawScene(FString("ControlRig"), DebugWorld->GetPhysicsScene()->GetDebugDrawScene());
+		Simulation->DebugDraw();
+	}
+#endif
+
+	if (!DI)
+	{
+		return;
+	}
+
+	// All rendering is done relative to the component, so convert the sim space (identity) into the
+	// component space.
+	const FTransform SpaceTransform = ConvertSimSpaceTransformToComponentSpace(SolverSettings, FTransform());
+
+	if (VisualizationSettings.bShowBodies)
+	{
+		if (CollisionActorHandle)
+		{
+			DrawActor(*DI, VisualizationSettings, SpaceTransform, *CollisionActorHandle);
+		}
+
+		for (const TPair<FRigComponentKey, FRigBodyRecord>& BodyRecordPair : BodyRecords)
+		{
+			const FRigBodyRecord& Record = BodyRecordPair.Value;
+			if (const ImmediatePhysics::FActorHandle* ActorHandle = Record.ActorHandle)
+			{
+				DrawActor(*DI, VisualizationSettings, SpaceTransform, *ActorHandle);
+			}
+		}
+	}
+
+	if (VisualizationSettings.bShowJoints)
+	{
+		for (const TPair<FRigComponentKey, FRigJointRecord>& JointRecordPair : JointRecords)
+		{
+			const FRigJointRecord& Record = JointRecordPair.Value;
+			if (const ImmediatePhysics::FJointHandle* JointHandle = Record.JointHandle)
+			{
+				DrawJoint(*DI, VisualizationSettings, SpaceTransform, *JointHandle, false);
+			}
+		}
+	}
+
+	if (VisualizationSettings.bShowControls)
+	{
+		for (const TPair<FRigComponentKey, FRigControlRecord>& ControlRecordPair : ControlRecords)
+		{
+			const FRigControlRecord& Record = ControlRecordPair.Value;
+			if (const ImmediatePhysics::FJointHandle* JointHandle = Record.JointHandle)
+			{
+				if (const Chaos::FPBDJointConstraintHandle* Constraint = JointHandle->GetConstraint())
+				{
+					if (IsConstraintDriveActive(*Constraint))
+					{
+						DrawJoint(*DI, VisualizationSettings, SpaceTransform, *JointHandle, true);
+					}
+				}
+			}
+		}
+	}
+
+	if (VisualizationSettings.bShowWorldOverlapBox)
+	{
+		if (WorldOverlapBox.IsValid)
+		{
+			// Remember that rendering is done in component space
+			FTransform BoxRelWorldTM(FQuat::Identity, WorldOverlapBox.GetCenter());
+			FTransform BoxRelComponentTM = BoxRelWorldTM.GetRelativeTransform(SimulationSpaceState.ComponentTM);
+			DI->DrawBox(
+				BoxRelComponentTM,
+				FTransform(FQuat::Identity, FVector::ZeroVector, WorldOverlapBox.GetSize()),
+				FColor::Green, 
+				VisualizationSettings.LineThickness);
+		}
+	}
+
+	if (VisualizationSettings.bShowWorldObjects)
+	{
+		for (TPair<uint32, FWorldObject>& WorldObjectPair : *WorldObjects.Get())
+		{
+			FWorldObject& WorldObject = WorldObjectPair.Value;
+			if (WorldObject.ActorHandle)
+			{
+				DrawActor(*DI, VisualizationSettings, SpaceTransform, *WorldObject.ActorHandle);
+			}
+		}
+	}
+
+	int32 ShowActiveOverride = CVarControlRigPhysicsShowActiveContactsOverride.GetValueOnAnyThread();
+	int32 ShowInactiveOverride = CVarControlRigPhysicsShowInactiveContactsOverride.GetValueOnAnyThread();
+
+	bool bShowActiveContacts = ShowActiveOverride < 0 ? VisualizationSettings.bShowActiveContacts : 
+		(ShowActiveOverride ? true : false);
+	bool bShowInactiveContacts = ShowInactiveOverride < 0 ? VisualizationSettings.bShowInactiveContacts :
+		(ShowInactiveOverride ? true : false);
+
+	if (bShowActiveContacts || bShowInactiveContacts)
+	{
+		Simulation->VisitCollisions(
+			[DI, &SpaceTransform, &VisualizationSettings, bShowActiveContacts, bShowInactiveContacts]
+			(const ImmediatePhysics::FSimulation::FCollisionData& Collision)
+			{
+				using namespace Chaos;
+
+				bool bIsActive = Collision.GetCollisionAccumulatedImpulse().IsZero();
+				if ((bIsActive && bShowActiveContacts) || (!bIsActive && bShowInactiveContacts))
+				{
+					float Size = 5.0f * VisualizationSettings.ShapeSize;
+
+					int32 NumManifoldPoints = Collision.GetNumManifoldPoints();
+					for (int32 PointIndex = 0; PointIndex != NumManifoldPoints; ++PointIndex)
+					{
+						Chaos::FRealSingle Depth;
+						Chaos::FVec3 PlaneNormal;
+						Chaos::FVec3 PointLocation;
+						Chaos::FVec3 PlaneLocation;
+						Collision.GetManifoldPointData(
+							PointIndex, Depth, PlaneNormal, PointLocation, PlaneLocation);
+
+						const FVec3 PointPlaneLocation = PointLocation -
+							FVec3::DotProduct(PointLocation - PlaneLocation, PlaneNormal) * PlaneNormal;
+
+						FMatrix Axes = FRotationMatrix::MakeFromZ(PlaneNormal);
+						FTransform TM(Axes);
+						TM.SetTranslation(PointPlaneLocation);
+
+						DI->DrawCircle(
+							SpaceTransform, TM, Size,
+							bIsActive ? RigPhysicsSolverDraw::ActiveContactColor : RigPhysicsSolverDraw::InactiveContactColor,
+							VisualizationSettings.LineThickness, VisualizationSettings.ShapeDetail);
+					}
+				}
+			}, Chaos::ECollisionVisitorFlags::VisitDefault);
+	}
+}
+
diff --git a/Engine/Plugins/Experimental/ControlRigPhysics/Source/ControlRigPhysics/Private/RigPhysicsSimulation_WorldObjects.cpp b/Engine/Plugins/Experimental/ControlRigPhysics/Source/ControlRigPhysics/Private/RigPhysicsSimulation_WorldObjects.cpp
new file mode 100644
--- /dev/null
+++ b/Engine/Plugins/Experimental/ControlRigPhysics/Source/ControlRigPhysics/Private/RigPhysicsSimulation_WorldObjects.cpp
@@ -0,0 +1,308 @@
+// Copyright Epic Games, Inc. All Rights Reserved.
+
+#include "RigPhysicsSimulation.h"
+#include "RigPhysicsSimulation_Space.inl"
+
+#include "Physics/ImmediatePhysics/ImmediatePhysicsActorHandle.h"
+#include "Physics/ImmediatePhysics/ImmediatePhysicsAdapters.h"
+#include "Physics/ImmediatePhysics/ImmediatePhysicsSimulation.h"
+#include "Physics/ImmediatePhysics/ImmediatePhysicsChaos/ImmediatePhysicsSimulation_Chaos.h"
+#include "Chaos/ParticleHandle.h"
+
+#include "Engine/OverlapResult.h"
+#include "Engine/World.h"
+#include "CollisionQueryParams.h"
+#include "CollisionShape.h"
+
+#include "Stats/Stats.h"
+#include "Stats/StatsHierarchical.h"
+
+#include "Components/InstancedStaticMeshComponent.h"
+
+//======================================================================================================================
+// Note that the simulation mutex will already be taken when we are called
+void FRigPhysicsSimulation::UpdateWorldObjectsPrePhysics(const FRigPhysicsSolverSettings& SolverSettings)
+{
+	DECLARE_SCOPE_HIERARCHICAL_COUNTER_FUNC()
+	QUICK_SCOPE_CYCLE_COUNTER(STAT_RigPhysics_UpdateWorldObjectsPrePhysics);
+
+	if (SolverSettings.WorldCollisionType == ERigPhysicsWorldCollisionType::None)
+	{
+		for (TPair<uint32, FWorldObject>& WorldObjectPair : *WorldObjects.Get())
+		{
+			FWorldObject& WorldObject = WorldObjectPair.Value;
+			if (WorldObject.ActorHandle)
+			{
+				UE_LOG(LogRigPhysics, Warning,
+					TEXT("Control rig %s - expiring object %s"),
+					*OwnerName.ToString(), *WorldObject.ActorHandle->GetName().ToString());
+
+				Simulation->DestroyActor(WorldObject.ActorHandle);
+			}
+		}
+		WorldObjects->Reset();
+	}
+	else
+	{
+		TArray<uint32> EntriesToRemove;
+		for (TPair<uint32, FWorldObject>& WorldObjectPair : *WorldObjects.Get())
+		{
+			FWorldObject& WorldObject = WorldObjectPair.Value;
+			if (WorldObject.GetExpired(UpdateCounter, SolverSettings.WorldCollisionExpiryFrames))
+			{
+				if (WorldObject.ActorHandle)
+				{
+					UE_LOG(LogRigPhysics, Log,
+						TEXT("Control rig %s - Expiring object %s"), 
+						*OwnerName.ToString(), *WorldObject.ActorHandle->GetName().ToString());
+
+					Simulation->DestroyActor(WorldObject.ActorHandle);
+					WorldObject.ActorHandle = nullptr;
+				}
+				EntriesToRemove.Add(WorldObjectPair.Key);
+			}
+		}
+		for (uint32 EntryToRemove : EntriesToRemove)
+		{
+			WorldObjects->Remove(EntryToRemove);
+		}
+	}
+
+	for (TPair<uint32, FWorldObject>& WorldObjectPair : *WorldObjects.Get())
+	{
+		FWorldObject& WorldObject = WorldObjectPair.Value;
+		FTransform SimSpaceTM = ConvertWorldTransformToSimSpace(SolverSettings, WorldObject.ComponentWorldTransform);
+		// Note that the code above should have ensured we don't have a null pointer
+		check(WorldObject.ActorHandle != nullptr);
+		if (WorldObject.LastSeenUpdateCounter == -1)
+		{
+			UE_LOG(LogRigPhysics, Log,
+				TEXT("Control rig %s - Adding object %s"), *OwnerName.ToString(),
+				*WorldObject.ActorHandle->GetName().ToString());
+
+			// We need to call InitWorldTransform (not just SetWorldTransform) here, to ensure the
+			// velocities are set to zero (though they should be after creation anyway), and that
+			// the object has its bounds set.
+			WorldObject.ActorHandle->InitWorldTransform(SimSpaceTM);
+			// If the object is found in the next overlap, then its update counter will be updated.
+			// However, in case it's not, we set it here too to make sure that it expires. 
+			WorldObject.LastSeenUpdateCounter = UpdateCounter;
+		}
+		else
+		{
+			WorldObject.ActorHandle->SetKinematicTarget(SimSpaceTM);
+		}
+	}
+}
+
+//======================================================================================================================
+// Note that the simulation mutex will already be taken when we are called
+void FRigPhysicsSimulation::UpdateWorldObjectsPostPhysics(
+	const UWorld* World, const FRigPhysicsSolverSettings& SolverSettings, const AActor* OwningActorPtr)
+{
+	DECLARE_SCOPE_HIERARCHICAL_COUNTER_FUNC()
+	QUICK_SCOPE_CYCLE_COUNTER(STAT_RigPhysics_UpdateWorldObjectsPostPhysics);
+
+	WorldOverlapBox = FBox();
+	if (SolverSettings.WorldCollisionType == ERigPhysicsWorldCollisionType::None)
+	{
+		return;
+	}
+
+	// Get the volume over which to do overlaps
+	const FTransform SpaceTransform = 
+		ConvertSimSpaceTransformToComponentSpace(SolverSettings, FTransform()) * SimulationSpaceState.ComponentTM;
+	for (TPair<FRigComponentKey, FRigBodyRecord>& BodyRecordPair : BodyRecords)
+	{
+		const FRigComponentKey& ComponentKey = BodyRecordPair.Key;
+		const FRigBodyRecord& Record = BodyRecordPair.Value;
+
+		if (const ImmediatePhysics::FActorHandle* ActorHandle = Record.ActorHandle)
+		{
+			if (!ActorHandle->GetIsKinematic())
+			{
+				const Chaos::FGeometryParticleHandle* GeometryParticleHandle = ActorHandle->GetParticle();
+				const FTransform ParticleTransform = GeometryParticleHandle->GetTransformXR() * SpaceTransform;
+
+				const auto& AABB = GeometryParticleHandle->LocalBounds();
+				if (!AABB.IsEmpty())
+				{
+					for (int32 Index = 0; Index != 8; ++Index)
+					{
+						WorldOverlapBox += ParticleTransform.TransformPosition(AABB.GetVertex(Index));
+					}
+				}
+			}
+		}
+	}
+
+	if (!WorldOverlapBox.IsValid)
+	{
+		return;
+	}
+
+	TWeakPtr<FTransactionallySafeCriticalSection> WeakSimulationMutex = SimulationMutex;
+
+	TWeakPtr<TMap<uint32, FWorldObject>> WeakWorldObjects = WorldObjects;
+	int64 UpdateCounterCopy = UpdateCounter;
+	const AActor* OwningActorPtrCopy = OwningActorPtr;
+	TSharedPtr<ImmediatePhysics::FSimulation> SimulationCopy = Simulation;
+
+	const TWeakObjectPtr<const UWorld> WeakWorld = World;
+
+	FVector NewHalfExtent = WorldOverlapBox.GetExtent() * SolverSettings.WorldCollisionBoundsExpansion;
+	FVector OverlapAABoxCenter = WorldOverlapBox.GetCenter();
+	WorldOverlapBox = FBox(OverlapAABoxCenter - NewHalfExtent, OverlapAABoxCenter + NewHalfExtent);
+	FCollisionShape OverlapAABox = FCollisionShape::MakeBox(NewHalfExtent);
+
+	FCollisionQueryParams QueryParams = FCollisionQueryParams(
+		SCENE_QUERY_STAT(ControlRigPhysicsFindGeometry), /*bTraceComplex=*/false);
+	QueryParams.MobilityType = EQueryMobilityType::Any;
+
+	FCollisionObjectQueryParams ObjectQueryParams;
+	switch (SolverSettings.WorldCollisionType)
+	{
+	case ERigPhysicsWorldCollisionType::Static:
+		ObjectQueryParams = FCollisionObjectQueryParams(FCollisionObjectQueryParams::InitType::AllStaticObjects); break;
+	case ERigPhysicsWorldCollisionType::Dynamic:
+		ObjectQueryParams = FCollisionObjectQueryParams(FCollisionObjectQueryParams::InitType::AllDynamicObjects); break;
+	case ERigPhysicsWorldCollisionType::All:
+		ObjectQueryParams = FCollisionObjectQueryParams(FCollisionObjectQueryParams::InitType::AllObjects); break;
+	default:
+		break;
+	}
+
+	// The actual overlap needs to be on the game thread, so we will run it in an async task.
+	// However, note that between now and then, we/our simulation etc might get destroyed so we must
+	// protect against that. Even the mutex guarding against access to the simulation may be released.
+	if (ObjectQueryParams.GetObjectTypesToQuery())
+	{
+		AsyncTask(ENamedThreads::GameThread,
+			[WeakSimulationMutex, WeakWorld, WeakWorldObjects, OverlapAABox, OverlapAABoxCenter, QueryParams, 
+			ObjectQueryParams, UpdateCounterCopy, OwningActorPtrCopy, SimulationCopy]()
+			{
+				if (TSharedPtr<FTransactionallySafeCriticalSection> SimulationMutex = WeakSimulationMutex.Pin())
+				{
+					UE::TScopeLock<FTransactionallySafeCriticalSection> Lock(*SimulationMutex.Get());
+
+					if (TStrongObjectPtr<const UWorld> World = WeakWorld.Pin())
+					{
+						if (TSharedPtr<TMap<uint32, FWorldObject>> WorldObjects = WeakWorldObjects.Pin())
+						{
+							TArray<FOverlapResult> Overlaps;
+							World->OverlapMultiByObjectType(
+								Overlaps, OverlapAABoxCenter, FQuat::Identity, ObjectQueryParams,
+								OverlapAABox, QueryParams);
+
+							for (const FOverlapResult& Overlap : Overlaps)
+							{
+								if (UPrimitiveComponent* OverlapComp = Overlap.GetComponent())
+								{
+									// For instanced static mesh components, each overlap result
+									// corresponds to a specific instance via Overlap.ItemIndex.
+									// The component-level BodyInstance has null BodySetup and no
+									// physics actor, so we must use the per-instance body instead.
+									int32 OverlapInstanceIndex = INDEX_NONE;
+									if (OverlapComp->IsA<UInstancedStaticMeshComponent>())
+									{
+										OverlapInstanceIndex = Overlap.ItemIndex;
+									}
+
+									// Use a per-instance composite key for ISM so each instance
+									// gets its own FWorldObject entry.
+									const uint32 Key = (OverlapInstanceIndex == INDEX_NONE)
+										? OverlapComp->GetUniqueID()
+										: HashCombine(OverlapComp->GetUniqueID(), (uint32)OverlapInstanceIndex);
+
+									FWorldObject* WorldObject = WorldObjects->Find(Key);
+									if (!WorldObject)
+									{
+										// New object - add it to the sim
+										if (OverlapComp->GetOwner() && OverlapComp->GetOwner() != OwningActorPtrCopy)
+										{
+											// Resolve the correct BodyInstance and transform.
+											// For ISM the component-level BodyInstance has null
+											// BodySetup, so use the per-instance body.
+											const FBodyInstance* BodyInstanceToUse = &OverlapComp->BodyInstance;
+
+											if (OverlapInstanceIndex != INDEX_NONE)
+											{
+												UInstancedStaticMeshComponent* ISMComp = Cast<UInstancedStaticMeshComponent>(OverlapComp);
+												BodyInstanceToUse = ISMComp	
+													? ISMComp->GetBodyInstance(NAME_None, false, OverlapInstanceIndex)
+													: nullptr;
+												if (BodyInstanceToUse == nullptr)
+												{
+													// Instance has no physics body
+													continue;
+												}
+											}
+
+											// Chaos holds scale separately to the world transform of the body
+											FTransform TransformToUse = BodyInstanceToUse->GetUnrealWorldTransform();
+											TransformToUse.SetScale3D(BodyInstanceToUse->Scale3D);
+
+											WorldObject = &WorldObjects->Add(Key, FWorldObject());
+
+											WorldObject->WorldPrimitiveComponent = OverlapComp;
+											WorldObject->InstanceBodyIndex = OverlapInstanceIndex;
+
+											// Note that we need to create the simulation actor here, even
+											// though we're in the game thread, so we can copy data out of the
+											// object. This should be safe because we only run this task after
+											// we've finished using the simulation in the worker thread.
+
+											WorldObject->ActorHandle = SimulationCopy->CreateActor(
+												ImmediatePhysics::MakeKinematicActorSetup(
+													BodyInstanceToUse, TransformToUse));
+
+#if WITH_EDITOR
+											WorldObject->ActorHandle->SetName(*OverlapComp->GetOwner()->GetActorLabel());
+#endif
+											SimulationCopy->AddToCollidingPairs(WorldObject->ActorHandle);
+
+											WorldObject->ComponentWorldTransform = TransformToUse;
+
+											// Flag that the simulation TM needs to be set, rather
+											// than using a kinematic target.
+											WorldObject->LastSeenUpdateCounter = -1;
+										}
+									}
+									else
+									{
+										if (TStrongObjectPtr<UPrimitiveComponent> WPC =
+											WorldObject->WorldPrimitiveComponent.Pin())
+										{
+											FBodyInstance* BodyInstanceToUse = &WPC->BodyInstance;
+
+											if (WorldObject->InstanceBodyIndex != INDEX_NONE)
+											{
+												UInstancedStaticMeshComponent* ISMComp = Cast<UInstancedStaticMeshComponent>(WPC.Get());
+												BodyInstanceToUse = ISMComp 
+													? ISMComp->GetBodyInstance(NAME_None, false, WorldObject->InstanceBodyIndex)
+													: nullptr;
+
+												if (BodyInstanceToUse == nullptr)
+												{
+													// Instance has no physics body
+													continue;
+												}
+											}
+
+											// Chaos holds scale separately to the world transform
+											FTransform TransformToUse = BodyInstanceToUse->GetUnrealWorldTransform();
+											TransformToUse.SetScale3D(BodyInstanceToUse->Scale3D);
+
+											WorldObject->ComponentWorldTransform = TransformToUse;
+											WorldObject->LastSeenUpdateCounter = UpdateCounterCopy;
+										}
+									}
+								}
+							}
+						}
+					}
+				}
+			});
+	}
+}
+
diff --git a/Engine/Plugins/Experimental/ControlRigPhysics/Source/ControlRigPhysics/Private/RigPhysicsSolverExecution.cpp b/Engine/Plugins/Experimental/ControlRigPhysics/Source/ControlRigPhysics/Private/RigPhysicsSolverExecution.cpp
new file mode 100644
--- /dev/null
+++ b/Engine/Plugins/Experimental/ControlRigPhysics/Source/ControlRigPhysics/Private/RigPhysicsSolverExecution.cpp
@@ -0,0 +1,224 @@
+// Copyright Epic Games, Inc. All Rights Reserved.
+
+#include "RigPhysicsSolverExecution.h"
+
+#include "RigPhysicsSimulation.h"
+#include "ControlRigPhysicsModule.h"
+
+#include "Rigs/RigHierarchyController.h"
+#include "Units/RigUnitContext.h"
+#include "Units/Execution/RigUnit_PrepareForExecution.h"
+#include "ControlRig.h"
+
+#include UE_INLINE_GENERATED_CPP_BY_NAME(RigPhysicsSolverExecution)
+
+TAutoConsoleVariable<int> CVarControlRigPhysicsVisualizationOverride(
+	TEXT("ControlRig.Physics.VisualizationOveride"), -1,
+	TEXT("-1 uses the Step Physics Solver setting, 0 forces visualization to be disabled, 1 forces it to be enabled."));
+
+TAutoConsoleVariable<int> CVarControlRigPhysicsEnableStepSolver(
+	TEXT("ControlRig.Physics.EnableStepSolver"), 1,
+	TEXT("Setting to zero disables stepping the solver"));
+
+
+//======================================================================================================================
+FRigUnit_AddPhysicsSolver_Execute()
+{
+	if (!ExecuteContext.Hierarchy)
+	{
+		return;
+	}
+
+	if (ExecuteContext.GetEventName() != FRigUnit_PrepareForExecution::EventName)
+	{
+		UE_CONTROLRIG_RIGUNIT_REPORT_ERROR(TEXT("AddPhysicsSolver can only be used during Setup"));
+		return;
+	}
+
+	if (URigHierarchyController* Controller = ExecuteContext.Hierarchy->GetController())
+	{
+		FRigHierarchyControllerInstructionBracket InstructionBracket(
+			Controller, ExecuteContext.GetInstructionIndex());
+		PhysicsSolverComponentKey = Controller->AddComponent(
+			FRigPhysicsSolverComponent::StaticStruct(), TEXT("PhysicsSolver"), Owner);
+		if (PhysicsSolverComponentKey.IsValid())
+		{
+			if (FRigPhysicsSolverComponent* Component = Cast<FRigPhysicsSolverComponent>(
+				ExecuteContext.Hierarchy->FindComponent(PhysicsSolverComponentKey)))
+			{
+				Component->SolverSettings = SolverSettings;
+				Component->SimulationSpaceSettings = SimulationSpaceSettings;
+			}
+		}
+	}
+}
+
+//======================================================================================================================
+FRigUnit_TrackInputPose_Execute()
+{
+	DECLARE_SCOPE_HIERARCHICAL_COUNTER_RIGUNIT();
+	if (!ExecuteContext.Hierarchy)
+	{
+		return;
+	}
+	if (FRigPhysicsSolverComponent* SolverComponent = Cast<FRigPhysicsSolverComponent>(
+		ExecuteContext.Hierarchy->FindComponent(PhysicsSolverComponentKey)))
+	{
+		if (bForceNumberOfFrames)
+		{
+			SolverComponent->TrackInputCounter = NumberOfFrames;
+		}
+		else
+		{
+			SolverComponent->TrackInputCounter = FMath::Max(SolverComponent->TrackInputCounter, NumberOfFrames);
+		}
+	}
+}
+
+//======================================================================================================================
+FRigUnit_StepPhysicsSolver_Execute()
+{
+	DECLARE_SCOPE_HIERARCHICAL_COUNTER_RIGUNIT();
+	if (!ExecuteContext.Hierarchy)
+	{
+		return;
+	}
+	if (CVarControlRigPhysicsEnableStepSolver.GetValueOnAnyThread() <= 0)
+	{
+		return;
+	}
+	URigHierarchy& Hierarchy = *ExecuteContext.Hierarchy;
+	if (UControlRig* ControlRig = Cast<UControlRig>(Hierarchy.GetOuter()))
+	{
+		if (FRigPhysicsSolverComponent* SolverComponent = Cast<FRigPhysicsSolverComponent>(
+			Hierarchy.FindComponent(PhysicsSolverComponentKey)))
+		{
+			FRigPhysicsSimulation* PhysicsSimulation = Cast<FRigPhysicsSimulation>(
+				ControlRig->GetPhysicsSimulation(PhysicsSolverComponentKey));
+
+			// Make the solver and instantiate if necessary
+			if (!PhysicsSimulation)
+			{
+				TSharedPtr<FRigPhysicsSimulation> NewPhysicsSimulation = MakeShared<FRigPhysicsSimulation>(
+					FName(ControlRig->GetName()));
+
+				if (ControlRig->RegisterPhysicsSimulation(NewPhysicsSimulation, PhysicsSolverComponentKey))
+				{
+					NewPhysicsSimulation->Instantiate(ExecuteContext, Hierarchy, *SolverComponent);
+					PhysicsSimulation = NewPhysicsSimulation.Get();
+				}
+			}
+
+			if (PhysicsSimulation)
+			{
+				const USceneComponent* OwningSceneComponent = ControlRig->GetOwningSceneComponent();
+				const AActor* OwningActorPtr = OwningSceneComponent ? OwningSceneComponent->GetOwner() : nullptr;
+
+				PhysicsSimulation->StepSimulation(
+					ControlRig->GetWorld(), OwningActorPtr, ExecuteContext, Hierarchy, *SolverComponent, 
+					DeltaTimeOverride, SimulationSpaceDeltaTimeOverride, Alpha);
+				int32 VisualizationOverride = CVarControlRigPhysicsVisualizationOverride.GetValueOnAnyThread();
+				if (VisualizationOverride == 1 || 
+					(VisualizationOverride < 0 && VisualizationSettings.bEnableVisualization))
+				{
+					PhysicsSimulation->Draw(
+						ExecuteContext.GetDrawInterface(), SolverComponent->SolverSettings,
+						VisualizationSettings, ExecuteContext.GetWorld());
+				}
+			}
+		}
+	}
+}
+
+//======================================================================================================================
+FRigUnit_SetPhysicsSolverSimulationSpaceSettings_Execute()
+{
+	DECLARE_SCOPE_HIERARCHICAL_COUNTER_RIGUNIT();
+	if (!ExecuteContext.Hierarchy)
+	{
+		return;
+	}
+	URigHierarchy& Hierarchy = *ExecuteContext.Hierarchy;
+	if (FRigPhysicsSolverComponent* SolverComponent = Cast<FRigPhysicsSolverComponent>(
+		Hierarchy.FindComponent(PhysicsSolverComponentKey)))
+	{
+		SolverComponent->SimulationSpaceSettings = SimulationSpaceSettings;
+	}
+}
+
+//======================================================================================================================
+FRigUnit_GetPhysicsSolverSimulationSpaceSettings_Execute()
+{
+	DECLARE_SCOPE_HIERARCHICAL_COUNTER_RIGUNIT();
+	if (!ExecuteContext.Hierarchy)
+	{
+		return;
+	}
+	URigHierarchy& Hierarchy = *ExecuteContext.Hierarchy;
+	if (const FRigPhysicsSolverComponent* SolverComponent = Cast<FRigPhysicsSolverComponent>(
+		Hierarchy.FindComponent(PhysicsSolverComponentKey)))
+	{
+		SimulationSpaceSettings = SolverComponent->SimulationSpaceSettings;
+	}
+}
+
+//======================================================================================================================
+FRigUnit_SetPhysicsSolverExternalVelocity_Execute()
+{
+	DECLARE_SCOPE_HIERARCHICAL_COUNTER_RIGUNIT();
+	if (!ExecuteContext.Hierarchy)
+	{
+		return;
+	}
+	URigHierarchy& Hierarchy = *ExecuteContext.Hierarchy;
+	if (FRigPhysicsSolverComponent* SolverComponent = Cast<FRigPhysicsSolverComponent>(
+		Hierarchy.FindComponent(PhysicsSolverComponentKey)))
+	{
+		SolverComponent->SimulationSpaceSettings.ExternalLinearVelocity = ExternalLinearVelocity;
+		SolverComponent->SimulationSpaceSettings.ExternalAngularVelocity = ExternalAngularVelocity;
+		SolverComponent->SimulationSpaceSettings.ExternalTurbulenceVelocity = ExternalTurbulenceVelocity;
+	}
+}
+
+//======================================================================================================================
+FRigUnit_GetPhysicsSolverSpaceData_Execute()
+{
+	DECLARE_SCOPE_HIERARCHICAL_COUNTER_RIGUNIT();
+	if (!ExecuteContext.Hierarchy)
+	{
+		return;
+	}
+	if (UControlRig* ControlRig = Cast<UControlRig>(ExecuteContext.Hierarchy->GetOuter()))
+	{
+		if (FRigPhysicsSolverComponent* SolverComponent = Cast<FRigPhysicsSolverComponent>(
+			ExecuteContext.Hierarchy->FindComponent(PhysicsSolverComponentKey)))
+		{
+			if (FRigPhysicsSimulation* PhysicsSimulation = Cast<FRigPhysicsSimulation>(
+				ControlRig->GetPhysicsSimulation(PhysicsSolverComponentKey)))
+			{
+				const FRigPhysicsSimulation::FSimulationSpaceData& SimulationSpaceData = 
+					PhysicsSimulation->GetSimulationSpaceData();
+
+				LinearVelocity = SimulationSpaceData.LinearVelocity;
+				AngularVelocity = SimulationSpaceData.AngularVelocity;
+				LinearAcceleration = SimulationSpaceData.LinearAcceleration;
+				AngularAcceleration = SimulationSpaceData.AngularAcceleration;
+				Gravity = SimulationSpaceData.Gravity;
+			}
+		}
+	}
+}
+
+//======================================================================================================================
+FRigUnit_HierarchySetPhysicsSolverAllowCCD_Execute()
+{
+	if (!ExecuteContext.Hierarchy)
+	{
+		return;
+	}
+	if (FRigPhysicsSolverComponent* SolverComponent = Cast<FRigPhysicsSolverComponent>(
+		ExecuteContext.Hierarchy->FindComponent(PhysicsSolverComponentKey)))
+	{
+		SolverComponent->SolverSettings.bAllowCCD = bAllowCCD;
+	}
+}
diff --git a/Engine/Plugins/Experimental/ControlRigPhysics/Source/ControlRigPhysics/Public/RigPhysicsBodyComponent.h b/Engine/Plugins/Experimental/ControlRigPhysics/Source/ControlRigPhysics/Public/RigPhysicsBodyComponent.h
--- a/Engine/Plugins/Experimental/ControlRigPhysics/Source/ControlRigPhysics/Public/RigPhysicsBodyComponent.h
+++ b/Engine/Plugins/Experimental/ControlRigPhysics/Source/ControlRigPhysics/Public/RigPhysicsBodyComponent.h
@@ -2,10 +2,6 @@
 #pragma once
 
 #include "RigPhysicsData.h"
-#include "PhysicsControlData.h"
-
-#include "Rigs/RigPhysics.h"
-#include "Rigs/RigHierarchyDefines.h"
 
 #include "RigPhysicsBodyComponent.generated.h"
 
diff --git a/Engine/Plugins/Experimental/ControlRigPhysics/Source/ControlRigPhysics/Public/RigPhysicsBodyExecution.h b/Engine/Plugins/Experimental/ControlRigPhysics/Source/ControlRigPhysics/Public/RigPhysicsBodyExecution.h
--- a/Engine/Plugins/Experimental/ControlRigPhysics/Source/ControlRigPhysics/Public/RigPhysicsBodyExecution.h
+++ b/Engine/Plugins/Experimental/ControlRigPhysics/Source/ControlRigPhysics/Public/RigPhysicsBodyExecution.h
@@ -3,12 +3,6 @@
 #pragma once
 
 #include "RigPhysicsExecution.h"
-#include "RigPhysicsBodyComponent.h"
-#include "RigPhysicsSolverComponent.h"
-#include "RigPhysicsSimulation.h"
-#include "PhysicsControlData.h"
-
-#include "Units/Execution/RigUnit_DynamicHierarchy.h"
 
 #include "RigPhysicsBodyExecution.generated.h"
 
diff --git a/Engine/Plugins/Experimental/ControlRigPhysics/Source/ControlRigPhysics/Public/RigPhysicsControlComponent.h b/Engine/Plugins/Experimental/ControlRigPhysics/Source/ControlRigPhysics/Public/RigPhysicsControlComponent.h
--- a/Engine/Plugins/Experimental/ControlRigPhysics/Source/ControlRigPhysics/Public/RigPhysicsControlComponent.h
+++ b/Engine/Plugins/Experimental/ControlRigPhysics/Source/ControlRigPhysics/Public/RigPhysicsControlComponent.h
@@ -5,8 +5,6 @@
 #include "PhysicsControlData.h"
 #include "RigPhysicsBodyComponent.h"
 
-#include "Rigs/RigPhysics.h"
-
 #include "RigPhysicsControlComponent.generated.h"
 
 #define UE_API CONTROLRIGPHYSICS_API
diff --git a/Engine/Plugins/Experimental/ControlRigPhysics/Source/ControlRigPhysics/Public/RigPhysicsControlExecution.h b/Engine/Plugins/Experimental/ControlRigPhysics/Source/ControlRigPhysics/Public/RigPhysicsControlExecution.h
--- a/Engine/Plugins/Experimental/ControlRigPhysics/Source/ControlRigPhysics/Public/RigPhysicsControlExecution.h
+++ b/Engine/Plugins/Experimental/ControlRigPhysics/Source/ControlRigPhysics/Public/RigPhysicsControlExecution.h
@@ -4,9 +4,6 @@
 
 #include "RigPhysicsExecution.h"
 #include "RigPhysicsControlComponent.h"
-#include "PhysicsControlData.h"
-
-#include "Units/Execution/RigUnit_DynamicHierarchy.h"
 
 #include "RigPhysicsControlExecution.generated.h"
 
diff --git a/Engine/Plugins/Experimental/ControlRigPhysics/Source/ControlRigPhysics/Public/RigPhysicsExecution.h b/Engine/Plugins/Experimental/ControlRigPhysics/Source/ControlRigPhysics/Public/RigPhysicsExecution.h
--- a/Engine/Plugins/Experimental/ControlRigPhysics/Source/ControlRigPhysics/Public/RigPhysicsExecution.h
+++ b/Engine/Plugins/Experimental/ControlRigPhysics/Source/ControlRigPhysics/Public/RigPhysicsExecution.h
@@ -4,8 +4,7 @@
 
 #include "RigPhysicsBodyComponent.h"
 #include "RigPhysicsSolverComponent.h"
-#include "RigPhysicsSimulation.h"
-#include "PhysicsControlData.h"
+//#include "RigPhysicsSimulation.h"
 
 #include "Units/Execution/RigUnit_DynamicHierarchy.h"
 
@@ -31,49 +30,6 @@ struct FRigUnit_PhysicsBase : public FRigUnit
 	GENERATED_BODY()
 };
 
-/**
- * Adds a new physics solver as a component on the owner element.
- * Note: This node only runs as part of the construction event.
- */
-USTRUCT(meta=(DisplayName="Spawn Physics Solver", Keywords="Add,Construction,Create,New,Simulation", Varying))
-struct FRigUnit_AddPhysicsSolver : public FRigUnit_PhysicsBaseMutable
-{
-	GENERATED_BODY()
-
-	FRigUnit_AddPhysicsSolver()
-	{
-		Owner.Type = ERigElementType::Bone;
-		// Default the material here to have friction and restitution. Then the interactions are
-		// easily adjusted on the dynamic bodies, because we prefer to combine by multiplying. If
-		// restitution was zero here, then setting restitution on the dynamic bodies would have no effect.
-		SolverSettings.Collision.Material.Friction = 1.0f;
-		SolverSettings.Collision.Material.Restitution = 1.0f;
-		SolverSettings.Collision.Material.FrictionCombineMode = ERigPhysicsCombineMode::Multiply;
-		SolverSettings.Collision.Material.RestitutionCombineMode = ERigPhysicsCombineMode::Multiply;
-	}
-
-	RIGVM_METHOD()
-	UE_API virtual void Execute() override;
-
-	// The owner of the newly created component (must be set/valid)
-	UPROPERTY(meta = (Input, BoneName))
-	FRigElementKey Owner;
-
-	// The component key of the solver that has been added
-	UPROPERTY(meta = (Output))
-	FRigComponentKey PhysicsSolverComponentKey;
-
-	// Settings for the physics solver that will be added
-	UPROPERTY(meta = (Input))
-	FRigPhysicsSolverSettings SolverSettings;
-
-	// Settings for the solver that apply to when it uses a simulation space other than "world".
-	// These typically relate to the movement of the simulation space itself, and how that is used
-	// by the solver.
-	UPROPERTY(meta = (Input))
-	FRigPhysicsSimulationSpaceSettings SimulationSpaceSettings;
-};
-
 /**
  * Instantiates all the objects in the physics world. Some properties can't be modified after this happens.
  * Note that it will happen automatically during the first simulation step if it hasn't been explicitly
@@ -99,169 +55,6 @@ struct FRigUnit_InstantiatePhysics : public FRigUnit_PhysicsBaseMutable
 	FRigComponentKey PhysicsSolverComponentKey;
 };
 
-/**
- * Steps the specified physics solver
- */
-USTRUCT(meta = (DisplayName = "Step Physics Solver", Keywords = "Simulate"))
-struct FRigUnit_StepPhysicsSolver : public FRigUnit_PhysicsBaseMutable
-{
-	GENERATED_BODY()
-
-	FRigUnit_StepPhysicsSolver()
-	{
-		PhysicsSolverComponentKey.ElementKey.Type = ERigElementType::Bone;
-		PhysicsSolverComponentKey.Name = FRigPhysicsSolverComponent::GetDefaultName();
-	}
-
-	RIGVM_METHOD()
-	UE_API virtual void Execute() override;
-
-	// The physics solver that should be stepped
-	UPROPERTY(meta = (Input))
-	FRigComponentKey PhysicsSolverComponentKey;
-	
-	// If this is zero, then the execute context time will be used. If this is positive then it will
-	// override the delta time. A negative value will prevent the solver from stepping, but there will
-	// still be update costs associated with the node.
-	UPROPERTY(meta = (Input))
-	float DeltaTimeOverride = 0.0f;
-
-	// If this is zero, then the simulation delta time will be used for evaluating movement of the
-	// simulation space. If this is positive then it will override. This may be needed if the
-	// component movement is being done in parallel, in which case you might need to pass in the
-	// previous time delta here.
-	UPROPERTY(meta = (Input, ClampMin = "0.0"))
-	float SimulationSpaceDeltaTimeOverride = 0.0f;
-
-	// How much of the simulation is combined with the input bone. This currently happens in
-	// component space. Note that the simulation will continue to run, even if alpha = 0
-	UPROPERTY(meta = (Input, ClampMin = "0.0", ClampMax = "1.0"))
-	float Alpha = 1.0f;
-
-	// Settings that specify how the solver state should be visualized during/after the step
-	UPROPERTY(meta = (Input))
-	FRigPhysicsVisualizationSettings VisualizationSettings;
-};
-
-// Sets the solver's simulation space settings
-USTRUCT(meta = (DisplayName = "Set Physics Solver Simulation Space Settings"))
-struct FRigUnit_SetPhysicsSolverSimulationSpaceSettings : public FRigUnit_PhysicsBaseMutable
-{
-	GENERATED_BODY()
-
-	FRigUnit_SetPhysicsSolverSimulationSpaceSettings()
-	{
-		PhysicsSolverComponentKey.ElementKey.Type = ERigElementType::Bone;
-		PhysicsSolverComponentKey.Name = FRigPhysicsSolverComponent::GetDefaultName();
-	}
-
-	RIGVM_METHOD()
-	UE_API virtual void Execute() override;
-
-	// The physics solver
-	UPROPERTY(meta = (Input))
-	FRigComponentKey PhysicsSolverComponentKey;
-
-	// The new simulation space settings
-	UPROPERTY(meta = (Input))
-	FRigPhysicsSimulationSpaceSettings SimulationSpaceSettings;
-};
-
-// Gets the solver's simulation space settings
-USTRUCT(meta = (DisplayName = "Get Physics Solver Simulation Space Settings"))
-struct FRigUnit_GetPhysicsSolverSimulationSpaceSettings : public FRigUnit_PhysicsBase
-{
-	GENERATED_BODY()
-
-	FRigUnit_GetPhysicsSolverSimulationSpaceSettings()
-	{
-		PhysicsSolverComponentKey.ElementKey.Type = ERigElementType::Bone;
-		PhysicsSolverComponentKey.Name = FRigPhysicsSolverComponent::GetDefaultName();
-	}
-
-	RIGVM_METHOD()
-	UE_API virtual void Execute() override;
-
-	// The physics solver
-	UPROPERTY(meta = (Input))
-	FRigComponentKey PhysicsSolverComponentKey;
-
-	// The current simulation space settings
-	UPROPERTY(meta = (Output))
-	FRigPhysicsSimulationSpaceSettings SimulationSpaceSettings;
-};
-
-// Sets the external velocity of the simulation - used for adding wind effects
-USTRUCT(meta = (DisplayName = "Set Physics Solver External Velocity"))
-struct FRigUnit_SetPhysicsSolverExternalVelocity : public FRigUnit_PhysicsBaseMutable
-{
-	GENERATED_BODY()
-
-	FRigUnit_SetPhysicsSolverExternalVelocity()
-	{
-		PhysicsSolverComponentKey.ElementKey.Type = ERigElementType::Bone;
-		PhysicsSolverComponentKey.Name = FRigPhysicsSolverComponent::GetDefaultName();
-	}
-
-	RIGVM_METHOD()
-	UE_API virtual void Execute() override;
-
-	// The physics solver
-	UPROPERTY(meta = (Input))
-	FRigComponentKey PhysicsSolverComponentKey;
-
-	// Additional velocity that is added to the component velocity so the simulation acts as if the
-	// actor is moving at speed, even when stationary. The vector is in world space. This could be
-	// used for wind effects etc. Typical values are similar to the velocity of the object or
-	// effect, and usually around or less than 1000 for characters/wind.
-	UPROPERTY(meta = (Input))
-	FVector ExternalLinearVelocity = FVector::ZeroVector;
-
-	// Additional angular velocity that is added to the component angular velocity. This can be used
-	// to make the simulation act as if the actor is rotating even when it is not. E.g., to apply
-	// physics to a character on a podium as the camera rotates around it, to emulate the podium
-	// itself rotating. Vector is in world space. Units are deg/s.
-	UPROPERTY(meta = (Input))
-	FVector ExternalAngularVelocity = FVector::ZeroVector;
-
-	// This will treat the external velocity like a wind field and add turbulence to it. Units are
-	// the same as velocity, so this is the approximate magnitude of the turbulence.
-	UPROPERTY(meta = (Input))
-	FVector ExternalTurbulenceVelocity = FVector::ZeroVector;
-};
-
-/**
- * Forces tracking of the input animation (on all physics bodies) for the next N frames
- */
-USTRUCT(meta = (DisplayName = "Track Input Pose", Keywords = "Reset,Simulate"))
-struct FRigUnit_TrackInputPose : public FRigUnit_PhysicsBaseMutable
-{
-	GENERATED_BODY()
-
-	FRigUnit_TrackInputPose()
-	{
-		PhysicsSolverComponentKey.ElementKey.Type = ERigElementType::Bone;
-		PhysicsSolverComponentKey.Name = FRigPhysicsSolverComponent::GetDefaultName();
-	}
-
-	RIGVM_METHOD()
-	UE_API virtual void Execute() override;
-
-	// The solver to relate this new physics element to
-	UPROPERTY(meta = (Input))
-	FRigComponentKey PhysicsSolverComponentKey;
-
-	// The number of frames to track the input pose for
-	UPROPERTY(meta = (Input, ClampMin = "0"))
-	int NumberOfFrames = 1;
-
-	// If true, then the number will be forced, potentially reducing the number. If false, then the
-	// NumberOfFrames will only be used to increase the number of frames remaining.
-	UPROPERTY(meta = (Input))
-	bool bForceNumberOfFrames = false;
-};
-
-
 // Adds a set of physics components including the body, joint and controls
 USTRUCT(meta = (DisplayName = "Spawn Physics Components", Keywords = "Add,Construction,Create,New,Body,Joint,Control", Varying))
 struct FRigUnit_AddPhysicsComponents : public FRigUnit_PhysicsBaseMutable
@@ -489,71 +282,4 @@ struct FRigUnit_HierarchyInstantiateFromPhysicsAsset : public FRigUnit_PhysicsBa
 	TArray<FRigComponentKey> ParentSpaceControlComponentKeys;
 };
 
-// Retrieves the simulation space data that were generated during the simulation step, so the values
-// returned will relate to the previous update if the solver has not yet been stepped.
-USTRUCT(meta = (DisplayName = "Get Physics Solver Space Data", Keywords = "Debug"))
-struct FRigUnit_GetPhysicsSolverSpaceData : public FRigUnit_PhysicsBase
-{
-	GENERATED_BODY()
-
-	FRigUnit_GetPhysicsSolverSpaceData()
-	{
-		PhysicsSolverComponentKey.ElementKey.Type = ERigElementType::Bone;
-		PhysicsSolverComponentKey.Name = FRigPhysicsSolverComponent::GetDefaultName();
-	}
-
-	RIGVM_METHOD()
-	UE_API virtual void Execute() override;
-
-	// The solver to relate this new physics element to
-	UPROPERTY(meta = (Input))
-	FRigComponentKey PhysicsSolverComponentKey;
-
-	// The velocity of the simulation space (in world space)
-	UPROPERTY(meta = (Output))
-	FVector LinearVelocity = FVector::ZeroVector;
-
-	// The angular velocity of the simulation space (in world space)
-	UPROPERTY(meta = (Output))
-	FVector AngularVelocity = FVector::ZeroVector;
-
-	// The linear acceleration of the simulation space (in world space)
-	UPROPERTY(meta = (Output))
-	FVector LinearAcceleration = FVector::ZeroVector;
-
-	// The angular acceleration of the simulation space (in world space)
-	UPROPERTY(meta = (Output))
-	FVector AngularAcceleration = FVector::ZeroVector;
-
-	// The gravitational acceleration that will be applied (in simulation space)
-	UPROPERTY(meta = (Output))
-	FVector Gravity = FVector::ZeroVector;
-};
-
-/**
- * Enables/disables support for CCD in the solver
- */
-USTRUCT(meta = (DisplayName = "Set Physics Solver Allow CCD", Varying))
-struct FRigUnit_HierarchySetPhysicsSolverAllowCCD : public FRigUnit_PhysicsBaseMutable
-{
-	GENERATED_BODY()
-
-	FRigUnit_HierarchySetPhysicsSolverAllowCCD()
-	{
-		PhysicsSolverComponentKey.ElementKey.Type = ERigElementType::Bone;
-		PhysicsSolverComponentKey.Name = FRigPhysicsSolverComponent::GetDefaultName();
-	}
-
-	RIGVM_METHOD()
-	UE_API virtual void Execute() override;
-
-	// The Physics Solver to be updated
-	UPROPERTY(meta = (Input))
-	FRigComponentKey PhysicsSolverComponentKey;
-
-	// Whether or not the solver allows CCD to be used. CCD also needs to be enabled on individual bodies
-	UPROPERTY(meta = (Input))
-	bool bAllowCCD = false;
-};
-
 #undef UE_API
diff --git a/Engine/Plugins/Experimental/ControlRigPhysics/Source/ControlRigPhysics/Public/RigPhysicsJointComponent.h b/Engine/Plugins/Experimental/ControlRigPhysics/Source/ControlRigPhysics/Public/RigPhysicsJointComponent.h
--- a/Engine/Plugins/Experimental/ControlRigPhysics/Source/ControlRigPhysics/Public/RigPhysicsJointComponent.h
+++ b/Engine/Plugins/Experimental/ControlRigPhysics/Source/ControlRigPhysics/Public/RigPhysicsJointComponent.h
@@ -5,11 +5,6 @@
 #include "RigPhysicsData.h"
 #include "RigPhysicsBodyComponent.h"
 
-#include "PhysicsControlData.h"
-
-#include "Rigs/RigPhysics.h"
-#include "Rigs/RigHierarchyDefines.h"
-
 #include "RigPhysicsJointComponent.generated.h"
 
 #define UE_API CONTROLRIGPHYSICS_API
diff --git a/Engine/Plugins/Experimental/ControlRigPhysics/Source/ControlRigPhysics/Public/RigPhysicsJointExecution.h b/Engine/Plugins/Experimental/ControlRigPhysics/Source/ControlRigPhysics/Public/RigPhysicsJointExecution.h
--- a/Engine/Plugins/Experimental/ControlRigPhysics/Source/ControlRigPhysics/Public/RigPhysicsJointExecution.h
+++ b/Engine/Plugins/Experimental/ControlRigPhysics/Source/ControlRigPhysics/Public/RigPhysicsJointExecution.h
@@ -3,12 +3,7 @@
 #pragma once
 
 #include "RigPhysicsExecution.h"
-#include "RigPhysicsBodyComponent.h"
 #include "RigPhysicsJointComponent.h"
-#include "RigPhysicsSolverComponent.h"
-#include "RigPhysicsSimulation.h"
-
-#include "Units/Execution/RigUnit_DynamicHierarchy.h"
 
 #include "RigPhysicsJointExecution.generated.h"
 
diff --git a/Engine/Plugins/Experimental/ControlRigPhysics/Source/ControlRigPhysics/Public/RigPhysicsSolverComponent.h b/Engine/Plugins/Experimental/ControlRigPhysics/Source/ControlRigPhysics/Public/RigPhysicsSolverComponent.h
--- a/Engine/Plugins/Experimental/ControlRigPhysics/Source/ControlRigPhysics/Public/RigPhysicsSolverComponent.h
+++ b/Engine/Plugins/Experimental/ControlRigPhysics/Source/ControlRigPhysics/Public/RigPhysicsSolverComponent.h
@@ -4,8 +4,6 @@
 
 #include "RigPhysicsData.h"
 
-#include "Rigs/RigPhysics.h"
-#include "Rigs/RigHierarchyDefines.h"
 #include "Rigs/RigHierarchyComponents.h"
 
 #include "RigPhysicsSolverComponent.generated.h"
diff --git a/Engine/Plugins/Experimental/ControlRigPhysics/Source/ControlRigPhysics/Public/RigPhysicsSolverExecution.h b/Engine/Plugins/Experimental/ControlRigPhysics/Source/ControlRigPhysics/Public/RigPhysicsSolverExecution.h
new file mode 100644
--- /dev/null
+++ b/Engine/Plugins/Experimental/ControlRigPhysics/Source/ControlRigPhysics/Public/RigPhysicsSolverExecution.h
@@ -0,0 +1,283 @@
+// Copyright Epic Games, Inc. All Rights Reserved.
+
+#pragma once
+
+#include "RigPhysicsExecution.h"
+
+#include "RigPhysicsSolverExecution.generated.h"
+
+#define UE_API CONTROLRIGPHYSICS_API
+
+/**
+ * Adds a new physics solver as a component on the owner element.
+ * Note: This node only runs as part of the construction event.
+ */
+USTRUCT(meta=(DisplayName="Spawn Physics Solver", Keywords="Add,Construction,Create,New,Simulation", Varying))
+struct FRigUnit_AddPhysicsSolver : public FRigUnit_PhysicsBaseMutable
+{
+	GENERATED_BODY()
+
+	FRigUnit_AddPhysicsSolver()
+	{
+		Owner.Type = ERigElementType::Bone;
+		// Default the material here to have friction and restitution. Then the interactions are
+		// easily adjusted on the dynamic bodies, because we prefer to combine by multiplying. If
+		// restitution was zero here, then setting restitution on the dynamic bodies would have no effect.
+		SolverSettings.Collision.Material.Friction = 1.0f;
+		SolverSettings.Collision.Material.Restitution = 1.0f;
+		SolverSettings.Collision.Material.FrictionCombineMode = ERigPhysicsCombineMode::Multiply;
+		SolverSettings.Collision.Material.RestitutionCombineMode = ERigPhysicsCombineMode::Multiply;
+	}
+
+	RIGVM_METHOD()
+	UE_API virtual void Execute() override;
+
+	// The owner of the newly created component (must be set/valid)
+	UPROPERTY(meta = (Input, BoneName))
+	FRigElementKey Owner;
+
+	// The component key of the solver that has been added
+	UPROPERTY(meta = (Output))
+	FRigComponentKey PhysicsSolverComponentKey;
+
+	// Settings for the physics solver that will be added
+	UPROPERTY(meta = (Input))
+	FRigPhysicsSolverSettings SolverSettings;
+
+	// Settings for the solver that apply to when it uses a simulation space other than "world".
+	// These typically relate to the movement of the simulation space itself, and how that is used
+	// by the solver.
+	UPROPERTY(meta = (Input))
+	FRigPhysicsSimulationSpaceSettings SimulationSpaceSettings;
+};
+
+/**
+ * Steps the specified physics solver
+ */
+USTRUCT(meta = (DisplayName = "Step Physics Solver", Keywords = "Simulate"))
+struct FRigUnit_StepPhysicsSolver : public FRigUnit_PhysicsBaseMutable
+{
+	GENERATED_BODY()
+
+	FRigUnit_StepPhysicsSolver()
+	{
+		PhysicsSolverComponentKey.ElementKey.Type = ERigElementType::Bone;
+		PhysicsSolverComponentKey.Name = FRigPhysicsSolverComponent::GetDefaultName();
+	}
+
+	RIGVM_METHOD()
+	UE_API virtual void Execute() override;
+
+	// The physics solver that should be stepped
+	UPROPERTY(meta = (Input))
+	FRigComponentKey PhysicsSolverComponentKey;
+	
+	// If this is zero, then the execute context time will be used. If this is positive then it will
+	// override the delta time. A negative value will prevent the solver from stepping, but there will
+	// still be update costs associated with the node.
+	UPROPERTY(meta = (Input))
+	float DeltaTimeOverride = 0.0f;
+
+	// If this is zero, then the simulation delta time will be used for evaluating movement of the
+	// simulation space. If this is positive then it will override. This may be needed if the
+	// component movement is being done in parallel, in which case you might need to pass in the
+	// previous time delta here.
+	UPROPERTY(meta = (Input, ClampMin = "0.0"))
+	float SimulationSpaceDeltaTimeOverride = 0.0f;
+
+	// How much of the simulation is combined with the input bone. This currently happens in
+	// component space. Note that the simulation will continue to run, even if alpha = 0
+	UPROPERTY(meta = (Input, ClampMin = "0.0", ClampMax = "1.0"))
+	float Alpha = 1.0f;
+
+	// Settings that specify how the solver state should be visualized during/after the step
+	UPROPERTY(meta = (Input))
+	FRigPhysicsVisualizationSettings VisualizationSettings;
+};
+
+// Sets the solver's simulation space settings
+USTRUCT(meta = (DisplayName = "Set Physics Solver Simulation Space Settings"))
+struct FRigUnit_SetPhysicsSolverSimulationSpaceSettings : public FRigUnit_PhysicsBaseMutable
+{
+	GENERATED_BODY()
+
+	FRigUnit_SetPhysicsSolverSimulationSpaceSettings()
+	{
+		PhysicsSolverComponentKey.ElementKey.Type = ERigElementType::Bone;
+		PhysicsSolverComponentKey.Name = FRigPhysicsSolverComponent::GetDefaultName();
+	}
+
+	RIGVM_METHOD()
+	UE_API virtual void Execute() override;
+
+	// The physics solver
+	UPROPERTY(meta = (Input))
+	FRigComponentKey PhysicsSolverComponentKey;
+
+	// The new simulation space settings
+	UPROPERTY(meta = (Input))
+	FRigPhysicsSimulationSpaceSettings SimulationSpaceSettings;
+};
+
+// Gets the solver's simulation space settings
+USTRUCT(meta = (DisplayName = "Get Physics Solver Simulation Space Settings"))
+struct FRigUnit_GetPhysicsSolverSimulationSpaceSettings : public FRigUnit_PhysicsBase
+{
+	GENERATED_BODY()
+
+	FRigUnit_GetPhysicsSolverSimulationSpaceSettings()
+	{
+		PhysicsSolverComponentKey.ElementKey.Type = ERigElementType::Bone;
+		PhysicsSolverComponentKey.Name = FRigPhysicsSolverComponent::GetDefaultName();
+	}
+
+	RIGVM_METHOD()
+	UE_API virtual void Execute() override;
+
+	// The physics solver
+	UPROPERTY(meta = (Input))
+	FRigComponentKey PhysicsSolverComponentKey;
+
+	// The current simulation space settings
+	UPROPERTY(meta = (Output))
+	FRigPhysicsSimulationSpaceSettings SimulationSpaceSettings;
+};
+
+// Sets the external velocity of the simulation - used for adding wind effects
+USTRUCT(meta = (DisplayName = "Set Physics Solver External Velocity"))
+struct FRigUnit_SetPhysicsSolverExternalVelocity : public FRigUnit_PhysicsBaseMutable
+{
+	GENERATED_BODY()
+
+	FRigUnit_SetPhysicsSolverExternalVelocity()
+	{
+		PhysicsSolverComponentKey.ElementKey.Type = ERigElementType::Bone;
+		PhysicsSolverComponentKey.Name = FRigPhysicsSolverComponent::GetDefaultName();
+	}
+
+	RIGVM_METHOD()
+	UE_API virtual void Execute() override;
+
+	// The physics solver
+	UPROPERTY(meta = (Input))
+	FRigComponentKey PhysicsSolverComponentKey;
+
+	// Additional velocity that is added to the component velocity so the simulation acts as if the
+	// actor is moving at speed, even when stationary. The vector is in world space. This could be
+	// used for wind effects etc. Typical values are similar to the velocity of the object or
+	// effect, and usually around or less than 1000 for characters/wind.
+	UPROPERTY(meta = (Input))
+	FVector ExternalLinearVelocity = FVector::ZeroVector;
+
+	// Additional angular velocity that is added to the component angular velocity. This can be used
+	// to make the simulation act as if the actor is rotating even when it is not. E.g., to apply
+	// physics to a character on a podium as the camera rotates around it, to emulate the podium
+	// itself rotating. Vector is in world space. Units are deg/s.
+	UPROPERTY(meta = (Input))
+	FVector ExternalAngularVelocity = FVector::ZeroVector;
+
+	// This will treat the external velocity like a wind field and add turbulence to it. Units are
+	// the same as velocity, so this is the approximate magnitude of the turbulence.
+	UPROPERTY(meta = (Input))
+	FVector ExternalTurbulenceVelocity = FVector::ZeroVector;
+};
+
+/**
+ * Forces tracking of the input animation (on all physics bodies) for the next N frames
+ */
+USTRUCT(meta = (DisplayName = "Set Solver Track Input Pose", Keywords = "Reset,Simulate"))
+struct FRigUnit_TrackInputPose : public FRigUnit_PhysicsBaseMutable
+{
+	GENERATED_BODY()
+
+	FRigUnit_TrackInputPose()
+	{
+		PhysicsSolverComponentKey.ElementKey.Type = ERigElementType::Bone;
+		PhysicsSolverComponentKey.Name = FRigPhysicsSolverComponent::GetDefaultName();
+	}
+
+	RIGVM_METHOD()
+	UE_API virtual void Execute() override;
+
+	// The solver to relate this new physics element to
+	UPROPERTY(meta = (Input))
+	FRigComponentKey PhysicsSolverComponentKey;
+
+	// The number of frames to track the input pose for
+	UPROPERTY(meta = (Input, ClampMin = "0"))
+	int NumberOfFrames = 1;
+
+	// If true, then the number will be forced, potentially reducing the number. If false, then the
+	// NumberOfFrames will only be used to increase the number of frames remaining.
+	UPROPERTY(meta = (Input))
+	bool bForceNumberOfFrames = false;
+};
+
+// Retrieves the simulation space data that were generated during the simulation step, so the values
+// returned will relate to the previous update if the solver has not yet been stepped.
+USTRUCT(meta = (DisplayName = "Get Physics Solver Space Data", Keywords = "Debug"))
+struct FRigUnit_GetPhysicsSolverSpaceData : public FRigUnit_PhysicsBase
+{
+	GENERATED_BODY()
+
+	FRigUnit_GetPhysicsSolverSpaceData()
+	{
+		PhysicsSolverComponentKey.ElementKey.Type = ERigElementType::Bone;
+		PhysicsSolverComponentKey.Name = FRigPhysicsSolverComponent::GetDefaultName();
+	}
+
+	RIGVM_METHOD()
+	UE_API virtual void Execute() override;
+
+	// The solver to relate this new physics element to
+	UPROPERTY(meta = (Input))
+	FRigComponentKey PhysicsSolverComponentKey;
+
+	// The velocity of the simulation space (in world space)
+	UPROPERTY(meta = (Output))
+	FVector LinearVelocity = FVector::ZeroVector;
+
+	// The angular velocity of the simulation space (in world space)
+	UPROPERTY(meta = (Output))
+	FVector AngularVelocity = FVector::ZeroVector;
+
+	// The linear acceleration of the simulation space (in world space)
+	UPROPERTY(meta = (Output))
+	FVector LinearAcceleration = FVector::ZeroVector;
+
+	// The angular acceleration of the simulation space (in world space)
+	UPROPERTY(meta = (Output))
+	FVector AngularAcceleration = FVector::ZeroVector;
+
+	// The gravitational acceleration that will be applied (in simulation space)
+	UPROPERTY(meta = (Output))
+	FVector Gravity = FVector::ZeroVector;
+};
+
+/**
+ * Enables/disables support for CCD in the solver
+ */
+USTRUCT(meta = (DisplayName = "Set Physics Solver Allow CCD", Varying))
+struct FRigUnit_HierarchySetPhysicsSolverAllowCCD : public FRigUnit_PhysicsBaseMutable
+{
+	GENERATED_BODY()
+
+	FRigUnit_HierarchySetPhysicsSolverAllowCCD()
+	{
+		PhysicsSolverComponentKey.ElementKey.Type = ERigElementType::Bone;
+		PhysicsSolverComponentKey.Name = FRigPhysicsSolverComponent::GetDefaultName();
+	}
+
+	RIGVM_METHOD()
+	UE_API virtual void Execute() override;
+
+	// The Physics Solver to be updated
+	UPROPERTY(meta = (Input))
+	FRigComponentKey PhysicsSolverComponentKey;
+
+	// Whether or not the solver allows CCD to be used. CCD also needs to be enabled on individual bodies
+	UPROPERTY(meta = (Input))
+	bool bAllowCCD = false;
+};
+
+#undef UE_API
