Control Rig Physics: Store the internal simulation/solver on the solver component. Also renames the internal object to RigPhysicsSolver, for consistency. Minor fixes etc from review.

#Synced-CL 52100224

diff --git a/Engine/Plugins/Animation/ControlRig/Source/ControlRig/Private/ControlRig.cpp b/Engine/Plugins/Animation/ControlRig/Source/ControlRig/Private/ControlRig.cpp
--- a/Engine/Plugins/Animation/ControlRig/Source/ControlRig/Private/ControlRig.cpp
+++ b/Engine/Plugins/Animation/ControlRig/Source/ControlRig/Private/ControlRig.cpp
@@ -798,10 +798,6 @@ void UControlRig::InitializeFromCDO()
 			}
 			URigHierarchy* Hierarchy = GetHierarchy();
 
-			// Physics Simulations are transient objects that only exist through construction so do
-			// not need to be initialized from the CDO
-			PhysicsSimulations.Empty();
-
 			// copy hierarchy
 			{
 				FRigHierarchyValidityBracket ValidityBracketA(Hierarchy);
@@ -1244,6 +1240,10 @@ bool UControlRig::Execute(const FName& InEventName)
 	{
 		check(Hierarchy);
 
+		// Reset our flag on construction - nodes etc are expected to set this if they contain
+		// state/would behave badly with some optional updates.
+		bContainsSimulation = false;
+
 		// construction mode means that we are running the construction event
 		// constantly for testing purposes.
 		const bool bConstructionModeEnabled = IsConstructionModeEnabled();
@@ -2087,10 +2087,6 @@ void UControlRig::Serialize(FArchive& Ar)
 		Ar << TempSolvers;
 		UE_RIGVM_ARCHIVETRACE_ENTRY(Ar, TEXT("PhysicsSolvers"));
 	}
-	else
-	{
-		PhysicsSimulations.Reset();
-	}
 
 	if(Ar.IsObjectReferenceCollector())
 	{
@@ -4169,34 +4165,6 @@ void UControlRig::GenerateRequiredPluginsData(FRigVMExtendedExecuteContext& Cont
 	}
 }
 
-FRigPhysicsSimulationBase* UControlRig::GetPhysicsSimulation(const FRigComponentKey& InComponentKey)
-{
-	TSharedPtr<FRigPhysicsSimulationBase>* PhysicsSimulation = PhysicsSimulations.Find(InComponentKey);
-	return PhysicsSimulation ? PhysicsSimulation->Get() : nullptr;
-}
-
-const FRigPhysicsSimulationBase* UControlRig::GetPhysicsSimulation(const FRigComponentKey& InComponentKey) const
-{
-	const TSharedPtr<FRigPhysicsSimulationBase>* PhysicsSimulation = PhysicsSimulations.Find(InComponentKey);
-	return PhysicsSimulation ? PhysicsSimulation->Get() : nullptr;
-}
-
-bool UControlRig::RegisterPhysicsSimulation(TSharedPtr<FRigPhysicsSimulationBase> InPhysicsSimulation, const FRigComponentKey& InComponentKey)
-{
-	if (IsReadOnly())
-	{
-		return false;
-	}
-
-	if(!InPhysicsSimulation.IsValid())
-	{
-		return false;
-	}
-	PhysicsSimulations.FindOrAdd(InComponentKey, InPhysicsSimulation);
-	InPhysicsSimulation->PhysicsSolverComponentKey = InComponentKey;
-	return true;
-}
-
 UControlRig::FPoseScope::FPoseScope(UControlRig* InControlRig, ERigElementType InFilter, const TArray<FRigElementKey>& InElements, const ERigTransformType::Type InTransformType)
 : ControlRig(InControlRig)
 , Filter(InFilter)
@@ -4305,9 +4273,4 @@ void UControlRig::SetReplay(UControlRigReplay* InReplay)
 	Replay = InReplay;
 }
 
-bool UControlRig::ContainsSimulation() const
-{
-	return !PhysicsSimulations.IsEmpty();
-}
-
 #undef LOCTEXT_NAMESPACE
diff --git a/Engine/Plugins/Animation/ControlRig/Source/ControlRig/Public/ControlRig.h b/Engine/Plugins/Animation/ControlRig/Source/ControlRig/Public/ControlRig.h
--- a/Engine/Plugins/Animation/ControlRig/Source/ControlRig/Public/ControlRig.h
+++ b/Engine/Plugins/Animation/ControlRig/Source/ControlRig/Public/ControlRig.h
@@ -805,32 +805,19 @@ protected:
 private:
 	float DebugBoneRadiusMultiplier;
 
-	// Physics Simulations (i.e. the simulations instantiated based on the Physics Solver components)
-	TMap<FRigComponentKey, TSharedPtr<FRigPhysicsSimulationBase>> PhysicsSimulations;
-
 public:
 
-	/**
-	 * Returns the physics simulation given the key of the component used to create it
-	 * @param InComponentKey The id identifying the physics simulation
-	 * @return The physics simulation (may be null)
-	 */
-	CONTROLRIG_API FRigPhysicsSimulationBase* GetPhysicsSimulation(const FRigComponentKey& InComponentKey);
+	// We no longer support the idea of ControlRig hosting simulations - that can be done by adding
+	// them to elements or components (etc). 
+	UE_DEPRECATED(5.8, "Control Rig core no longer supports owning simulations")
+	FRigPhysicsSimulationBase* GetPhysicsSimulation(const FRigComponentKey& InComponentKey) { return nullptr ;}
 
-	/**
-	 * Returns the physics simulation given the key of the component used to create it
-	 * @param InComponentKey The id identifying the physics simulation
-	 * @return The physics simulation (may be null)
-	 */
-	CONTROLRIG_API const FRigPhysicsSimulationBase* GetPhysicsSimulation(const FRigComponentKey& InComponentKey) const;
+	UE_DEPRECATED(5.8, "Control Rig core no longer supports owning simulations")
+	const FRigPhysicsSimulationBase* GetPhysicsSimulation(const FRigComponentKey& InComponentKey) const { return nullptr; }
 
-	/**
-	 * Registers and stores (taking ownership) of a physics simulation. The simulation will be 
-	 * responsible for the low-level storage and simulation of objects. Returns false if the world 
-	 * can't be registered - e.g. if there is already a simulation registered with the key.
-	 */
-	CONTROLRIG_API bool RegisterPhysicsSimulation(
-		TSharedPtr<FRigPhysicsSimulationBase> PhysicsSimulation, const FRigComponentKey& InComponentKey);
+	UE_DEPRECATED(5.8, "Control Rig core no longer supports owning simulations")
+	bool RegisterPhysicsSimulation(
+		TSharedPtr<FRigPhysicsSimulationBase> PhysicsSimulation, const FRigComponentKey& InComponentKey) { return false; }
 
 #if WITH_EDITOR	
 
@@ -997,9 +984,17 @@ public:
 	CONTROLRIG_API bool IsReplayEnabled() const;
 	CONTROLRIG_API void SetReplay(UControlRigReplay* InReplay);
 
-	CONTROLRIG_API bool ContainsSimulation() const;
+	// Nodes/components etc that contain "simulation" (i.e. contain significant amounts of state,
+	// and won't behave nicely when given extra "updates") should set this on the control rig during
+	// construction or during execution. It will be used to prevent additional updates of the
+	// control rig that are there primarily to make it feel responsive when editing. Note that the
+	// flag will be reset automatically before construction etc.
+	void SetContainsSimulation() { bContainsSimulation = true; }
+	bool ContainsSimulation() const { return bContainsSimulation; }
 
 private:
+	bool bContainsSimulation = false;
+
 	int32 ReplayTimeIndex;
 	TObjectPtr<UControlRigReplay> Replay;
 
diff --git a/Engine/Plugins/Animation/ControlRig/Source/ControlRig/Public/Rigs/RigPhysics.h b/Engine/Plugins/Animation/ControlRig/Source/ControlRig/Public/Rigs/RigPhysics.h
--- a/Engine/Plugins/Animation/ControlRig/Source/ControlRig/Public/Rigs/RigPhysics.h
+++ b/Engine/Plugins/Animation/ControlRig/Source/ControlRig/Public/Rigs/RigPhysics.h
@@ -8,10 +8,9 @@
 
 #define UE_API CONTROLRIG_API
 
-// Base class for physics simulations registered with the ControlRig. The concrete class will be
-// instantiated by a plugin, so there is casting support to ensure that the correct concrete class
-// is being used.
-USTRUCT()
+// Control Rig core no longer supports simulations. Control Rig Physics (for example) now stores it
+// in the solver component.
+USTRUCT(meta = (Deprecated = "5.8"))
 struct FRigPhysicsSimulationBase
 {
 	GENERATED_BODY()
@@ -30,10 +29,7 @@ struct FRigPhysicsSimulationBase
 
 	virtual ~FRigPhysicsSimulationBase() = default;
 
-	// Used to support casting from the base class to the concrete class
 	const UScriptStruct* Type;
-
-	// The component key will be set when the concrete class is registered with the control rig
 	FRigComponentKey PhysicsSolverComponentKey;
 };
 
diff --git a/Engine/Plugins/Experimental/ControlRigPhysics/Source/ControlRigPhysics/Private/RigPhysicsExecution.cpp b/Engine/Plugins/Experimental/ControlRigPhysics/Source/ControlRigPhysics/Private/RigPhysicsExecution.cpp
--- a/Engine/Plugins/Experimental/ControlRigPhysics/Source/ControlRigPhysics/Private/RigPhysicsExecution.cpp
+++ b/Engine/Plugins/Experimental/ControlRigPhysics/Source/ControlRigPhysics/Private/RigPhysicsExecution.cpp
@@ -3,7 +3,7 @@
 #include "RigPhysicsExecution.h"
 #include "RigPhysicsJointComponent.h"
 #include "RigPhysicsControlComponent.h"
-#include "RigPhysicsSimulation.h"
+#include "RigPhysicsSolver.h"
 #include "ControlRigPhysicsModule.h"
 
 #include "Rigs/RigHierarchyController.h"
@@ -18,31 +18,6 @@
 
 #include UE_INLINE_GENERATED_CPP_BY_NAME(RigPhysicsExecution)
 
-//======================================================================================================================
-FRigUnit_InstantiatePhysics_Execute()
-{
-	DECLARE_SCOPE_HIERARCHICAL_COUNTER_RIGUNIT();
-	if (!ExecuteContext.Hierarchy)
-	{
-		return;
-	}
-	URigHierarchy& Hierarchy = *ExecuteContext.Hierarchy;
-	if (UControlRig* ControlRig = Cast<UControlRig>(Hierarchy.GetOuter()))
-	{
-		if (FRigPhysicsSolverComponent* SolverComponent = Cast<FRigPhysicsSolverComponent>(
-			ExecuteContext.Hierarchy->FindComponent(PhysicsSolverComponentKey)))
-		{
-			TSharedPtr<FRigPhysicsSimulation> PhysicsSimulation = MakeShared<FRigPhysicsSimulation>(
-				FName(ControlRig->GetName()));
-
-			if (ControlRig->RegisterPhysicsSimulation(PhysicsSimulation, PhysicsSolverComponentKey))
-			{
-				PhysicsSimulation->Instantiate(ExecuteContext, Hierarchy, *SolverComponent);
-			}
-		}
-	}
-}
-
 //======================================================================================================================
 FRigUnit_AddPhysicsComponents_Execute()
 {
diff --git a/Engine/Plugins/Experimental/ControlRigPhysics/Source/ControlRigPhysics/Private/RigPhysicsSimulation.cpp b/Engine/Plugins/Experimental/ControlRigPhysics/Source/ControlRigPhysics/Private/RigPhysicsSimulation.cpp
deleted file mode 100644
diff --git a/Engine/Plugins/Experimental/ControlRigPhysics/Source/ControlRigPhysics/Private/RigPhysicsSimulation.h b/Engine/Plugins/Experimental/ControlRigPhysics/Source/ControlRigPhysics/Private/RigPhysicsSimulation.h
deleted file mode 100644
diff --git a/Engine/Plugins/Experimental/ControlRigPhysics/Source/ControlRigPhysics/Private/RigPhysicsSimulation_Space.cpp b/Engine/Plugins/Experimental/ControlRigPhysics/Source/ControlRigPhysics/Private/RigPhysicsSimulation_Space.cpp
deleted file mode 100644
diff --git a/Engine/Plugins/Experimental/ControlRigPhysics/Source/ControlRigPhysics/Private/RigPhysicsSimulation_Space.inl b/Engine/Plugins/Experimental/ControlRigPhysics/Source/ControlRigPhysics/Private/RigPhysicsSimulation_Space.inl
deleted file mode 100644
diff --git a/Engine/Plugins/Experimental/ControlRigPhysics/Source/ControlRigPhysics/Private/RigPhysicsSimulation_Visualization.cpp b/Engine/Plugins/Experimental/ControlRigPhysics/Source/ControlRigPhysics/Private/RigPhysicsSimulation_Visualization.cpp
deleted file mode 100644
diff --git a/Engine/Plugins/Experimental/ControlRigPhysics/Source/ControlRigPhysics/Private/RigPhysicsSimulation_WorldObjects.cpp b/Engine/Plugins/Experimental/ControlRigPhysics/Source/ControlRigPhysics/Private/RigPhysicsSimulation_WorldObjects.cpp
deleted file mode 100644
diff --git a/Engine/Plugins/Experimental/ControlRigPhysics/Source/ControlRigPhysics/Private/RigPhysicsSolver.cpp b/Engine/Plugins/Experimental/ControlRigPhysics/Source/ControlRigPhysics/Private/RigPhysicsSolver.cpp
new file mode 100644
--- /dev/null
+++ b/Engine/Plugins/Experimental/ControlRigPhysics/Source/ControlRigPhysics/Private/RigPhysicsSolver.cpp
@@ -0,0 +1,1949 @@
+// Copyright Epic Games, Inc. All Rights Reserved.
+
+#include "RigPhysicsSolver.h"
+#include "RigPhysicsSolver_Space.inl"
+
+#include "RigPhysicsBodyComponent.h"
+#include "RigPhysicsJointComponent.h"
+#include "RigPhysicsControlComponent.h"
+#include "RigPhysicsSolverComponent.h"
+
+#include "AnimNode_RigidBodyWithControl.h"
+#include "PhysicsControlPoseData.h"
+#include "PhysicsControlHelpers.h"
+
+#include "Rigs/RigHierarchy.h"
+#include "Rigs/RigHierarchyController.h"
+#include "Units/RigUnitContext.h"
+#include "RigVMCore/RigVMExecuteContext.h"
+
+#include "Engine/World.h"
+
+#include "Physics/ImmediatePhysics/ImmediatePhysicsActorHandle.h"
+#include "Physics/ImmediatePhysics/ImmediatePhysicsAdapters.h"
+#include "Physics/ImmediatePhysics/ImmediatePhysicsJointHandle.h"
+#include "Physics/ImmediatePhysics/ImmediatePhysicsSimulation.h"
+#include "Physics/ImmediatePhysics/ImmediatePhysicsChaos/ImmediatePhysicsSimulation_Chaos.h"
+
+#include "Physics/Experimental/PhysScene_Chaos.h"
+#include "Physics/Experimental/ChaosInterfaceUtils.h"
+
+#include "Chaos/ParticleHandle.h"
+#include "Chaos/PBDJointConstraints.h"
+#include "Chaos/ImplicitObject.h"
+#include "Chaos/ShapeInstance.h"
+#include "Chaos/Capsule.h"
+#include "Chaos/ChaosScene.h"
+#include "Chaos/PBDJointConstraintUtilities.h"
+#include "Chaos/Evolution/SimulationSpace.h"
+#include "Chaos/PBDCollisionConstraints.h"
+#include "Chaos/ChaosDebugNameDefines.h"
+
+#include "PhysicsEngine/BodyInstance.h"
+#include "PhysicsEngine/ConstraintInstance.h"
+#include "PhysicsEngine/BodySetup.h"
+
+#include "Stats/Stats.h"
+
+#if WITH_CHAOS_VISUAL_DEBUGGER
+#include "ChaosVDRuntimeModule.h"
+#endif
+
+#include "Logging/LogMacros.h"
+
+#include UE_INLINE_GENERATED_CPP_BY_NAME(RigPhysicsSolver)
+
+DEFINE_LOG_CATEGORY(LogRigPhysics);
+
+TAutoConsoleVariable<float> CVarControlRigPhysicsFixedTimeStepOverride(
+	TEXT("ControlRig.Physics.FixedTimeStepOverride"), -1.0f,
+	TEXT("-1.0 disables the override, so the timestep authored in the simulation settings will be used (which may or may not imply a fixed timestep). A value of 0 forces a variable timestep to be used. A +ve value is used to specify a fixed timestep."));
+
+TAutoConsoleVariable<int> CVarControlRigPhysicsMaxTimeStepsOverride(
+	TEXT("ControlRig.Physics.MaxTimeStepsOverride"), -1,
+	TEXT("-1 disables the override, so the max timesteps authored in the simulation settings will be used. A +ve value is used to specify the maximum number of timesteps."));
+
+TAutoConsoleVariable<float> CVarControlRigPhysicsMaxDeltaTimeOverride(
+	TEXT("ControlRig.Physics.MaxDeltaTimeOverride"), -1,
+	TEXT("-1 disables the override, so the max delta time authored in the simulation settings will be used. A +ve value is used to specify the maximum delta time."));
+
+constexpr int32 ConstraintChildIndex = 0;
+constexpr int32 ConstraintParentIndex = 1;
+
+//======================================================================================================================
+FRigPhysicsSolver::FRigPhysicsSolver(const FRigComponentKey& InSolverComponentKey)
+	: PhysicsSolverComponentKey(InSolverComponentKey)
+{
+	WorldObjects = MakeShared<TMap<uint32, FWorldObject>>();
+	SimulationMutex = MakeShared<FTransactionallySafeCriticalSection>();
+}
+
+//======================================================================================================================
+bool FRigPhysicsSolver::ShouldComponentBeInSimulation(
+	const URigHierarchy&              Hierarchy, 
+	const FRigPhysicsSolverComponent& SolverComponent, 
+	const FRigComponentKey&           ComponentKey) const
+{
+	const FRigPhysicsBodyComponent* PhysicsBodyComponent = Cast<FRigPhysicsBodyComponent>(
+		Hierarchy.FindComponent(ComponentKey));
+
+	if (!PhysicsBodyComponent)
+	{
+		return false;
+	}
+
+	check(PhysicsSolverComponentKey == SolverComponent.GetKey());
+	if (PhysicsBodyComponent->BodySolverSettings.PhysicsSolverComponentKey == PhysicsSolverComponentKey)
+	{
+		return true;
+	}
+
+	if (!SolverComponent.SolverSettings.bAutomaticallyAddPhysicsBodyComponents)
+	{
+		return false;
+	}
+
+	if (PhysicsBodyComponent->BodySolverSettings.bUseAutomaticSolver)
+	{
+		FRigElementKey ElementKey = PhysicsBodyComponent->GetElementKey();
+		while (ElementKey.IsValid())
+		{
+			for (FRigComponentKey CK : Hierarchy.GetComponentKeys(ElementKey))
+			{
+				if (CK == PhysicsSolverComponentKey)
+				{
+					return true;
+				}
+			}
+			// Note that getting the parent of an element at the root doesn't return the top-level element
+			ElementKey = Hierarchy.GetFirstParent(ElementKey);
+		}
+	}
+
+	return false;
+}
+
+
+//======================================================================================================================
+void FRigPhysicsSolver::InitialiseSimulation()
+{
+	DestroyPhysicsSimulation();
+
+	Simulation = MakeShared<ImmediatePhysics::FSimulation>();
+
+#if CHAOS_SOLVER_DEBUG_NAME
+	Simulation->SetDebugName(FName(TEXT("ControlRigPhysics")));
+#endif
+
+#if WITH_CHAOS_VISUAL_DEBUGGER
+	Simulation->GetChaosVDContextData().Id = FChaosVDRuntimeModule::Get().GenerateUniqueID();
+	Simulation->GetChaosVDContextData().Type = static_cast<int32>(EChaosVDContextType::Solver);
+#endif
+
+	// This is needed so that when using a fixed timestep, velocities are rewound as well as
+	// positions. This is not only more accurate, but it's needed in order to get soft constraint
+	// behavior (in particular, for controls) that behave fairly independently of the control-rig
+	// tick rate.
+	Simulation->SetRewindVelocities(true);
+
+	// Always create a world actor at the origin, for attaching controls to. 
+	SimulationActorHandle = CreateBody(
+		FName(TEXT("Simulation")), FRigPhysicsCollision(), nullptr, nullptr, FTransform());
+}
+
+//======================================================================================================================
+void FRigPhysicsSolver::InitialiseControlRecords(
+	const URigHierarchy& Hierarchy, const FRigPhysicsSolverComponent& SolverComponent)
+{
+	ensure(ControlRecords.IsEmpty());
+
+	TArray<FRigComponentKey> AllComponentKeys = Hierarchy.GetAllComponentKeys();
+
+	for (FRigComponentKey ComponentKey : AllComponentKeys)
+	{
+		if (const FRigPhysicsControlComponent* ControlComponent = Cast<FRigPhysicsControlComponent>(
+			Hierarchy.FindComponent(ComponentKey)))
+		{
+			// The authored body components may be blank, in which case we need to find the automatic ones.
+			FRigControlRecord ControlRecord;
+			ControlRecord.ParentBodyComponentKey = ControlComponent->ParentBodyComponentKey;
+			ControlRecord.ChildBodyComponentKey = ControlComponent->ChildBodyComponentKey;
+
+			// Automate the child
+			if (!ControlRecord.ChildBodyComponentKey.IsValid())
+			{
+				TArray<FRigComponentKey> SiblingComponentKeys = Hierarchy.GetComponentKeys(ComponentKey.ElementKey);
+				for (FRigComponentKey SiblingComponentKey : SiblingComponentKeys)
+				{
+					if (ShouldComponentBeInSimulation(Hierarchy, SolverComponent, SiblingComponentKey))
+					{
+						ControlRecord.ChildBodyComponentKey = SiblingComponentKey;
+						break;
+					}
+				}
+			}
+
+			if (ControlComponent->bUseParentBodyAsDefault)
+			{
+				if (!ControlRecord.ParentBodyComponentKey.IsValid())
+				{
+					FRigElementKey ParentElementKey = Hierarchy.GetFirstParent(ComponentKey.ElementKey);
+					TArray<FRigComponentKey> ParentComponentKeys = Hierarchy.GetComponentKeys(ParentElementKey);
+					for (FRigComponentKey ParentComponentKey : ParentComponentKeys)
+					{
+						if (ShouldComponentBeInSimulation(Hierarchy, SolverComponent, ParentComponentKey))
+						{
+							ControlRecord.ParentBodyComponentKey = ParentComponentKey;
+							break;
+						}
+					}
+				}
+			}
+
+			if (ShouldComponentBeInSimulation(Hierarchy, SolverComponent, ControlRecord.ChildBodyComponentKey))
+			{
+				// Here, an invalid parent component key indicates a sim-space control. 
+				if (!ControlRecord.ParentBodyComponentKey.IsValid() ||
+					ShouldComponentBeInSimulation(Hierarchy, SolverComponent, ControlRecord.ParentBodyComponentKey))
+				{
+					// Just make the record for now - it will be instantiated later
+					ControlRecords.Add(ComponentKey, ControlRecord);
+				}
+			}
+		}
+	}
+}
+
+//======================================================================================================================
+void FRigPhysicsSolver::InitialiseJointRecords(
+	const URigHierarchy& Hierarchy, const FRigPhysicsSolverComponent& SolverComponent)
+{
+	ensure(JointRecords.IsEmpty());
+
+	FRigComponentKey SolverComponentKey = SolverComponent.GetKey();
+
+	TArray<FRigComponentKey> AllComponentKeys = Hierarchy.GetAllComponentKeys();
+
+	for (FRigComponentKey ComponentKey : AllComponentKeys)
+	{
+		if (const FRigPhysicsJointComponent* JointComponent = Cast<FRigPhysicsJointComponent>(
+			Hierarchy.FindComponent(ComponentKey)))
+		{
+			// The authored body components may be blank, in which case we need to find the automatic ones.
+			FRigJointRecord JointRecord;
+			JointRecord.ParentBodyComponentKey = JointComponent->ParentBodyComponentKey;
+			JointRecord.ChildBodyComponentKey = JointComponent->ChildBodyComponentKey;
+
+			if (!JointRecord.ChildBodyComponentKey.IsValid())
+			{
+				TArray<FRigComponentKey> SiblingComponentKeys = Hierarchy.GetComponentKeys(ComponentKey.ElementKey);
+				for (FRigComponentKey SiblingComponentKey : SiblingComponentKeys)
+				{
+					if (ShouldComponentBeInSimulation(Hierarchy, SolverComponent, SiblingComponentKey))
+					{
+						JointRecord.ChildBodyComponentKey = SiblingComponentKey;
+						break;
+					}
+				}
+			}
+
+			if (!JointRecord.ParentBodyComponentKey.IsValid())
+			{
+				FRigElementKey ParentElementKey = Hierarchy.GetFirstParent(ComponentKey.ElementKey);
+				TArray<FRigComponentKey> ParentComponentKeys = Hierarchy.GetComponentKeys(ParentElementKey);
+				for (FRigComponentKey ParentComponentKey : ParentComponentKeys)
+				{
+					if (ShouldComponentBeInSimulation(Hierarchy, SolverComponent, ParentComponentKey))
+					{
+						JointRecord.ParentBodyComponentKey = ParentComponentKey;
+						break;
+					}
+				}
+			}
+
+			if (ShouldComponentBeInSimulation(Hierarchy, SolverComponent, JointRecord.ChildBodyComponentKey))
+			{
+				if (ShouldComponentBeInSimulation(Hierarchy, SolverComponent, JointRecord.ParentBodyComponentKey))
+				{
+					// Just make the record for now - it will be instantiated later
+					JointRecords.Add(ComponentKey, JointRecord);
+				}
+			}
+		}
+	}
+}
+
+
+//======================================================================================================================
+void FRigPhysicsSolver::InitialiseBodyRecords(
+	const URigHierarchy& Hierarchy, const FRigPhysicsSolverComponent& SolverComponent)
+{
+	ensure(BodyRecords.IsEmpty());
+
+	FRigComponentKey SolverComponentKey = SolverComponent.GetKey();
+
+	TArray<FRigComponentKey> AllComponentKeys = Hierarchy.GetAllComponentKeys();
+
+	// All the components in this simulation
+	TArray<FRigComponentKey> UnsortedBodyComponentKeys;
+
+	for (FRigComponentKey ComponentKey : AllComponentKeys)
+	{
+		if (ShouldComponentBeInSimulation(Hierarchy, SolverComponent, ComponentKey))
+		{
+			// Just make the record for now - it will be instantiated later
+			BodyRecords.Add(ComponentKey);
+			UnsortedBodyComponentKeys.Add(ComponentKey);
+		}
+	}
+
+	// Sort the component keys according to the traversal of their element (i.e. from root to leaf)
+	SortedBodyComponentKeys.Empty(UnsortedBodyComponentKeys.Num());
+	Hierarchy.Traverse([this, &UnsortedBodyComponentKeys](FRigBaseElement* Element, bool& bContinue)
+		{
+			const FRigElementKey& Key = Element->GetKey();
+			for (const FRigComponentKey& ComponentKey : UnsortedBodyComponentKeys)
+			{
+				if (ComponentKey.ElementKey == Key)
+				{
+					SortedBodyComponentKeys.Push(ComponentKey);
+				}
+			}
+		});
+}
+
+//======================================================================================================================
+void FRigPhysicsSolver::DestroyPhysicsSimulation()
+{
+	if (Simulation.IsValid())
+	{
+		for (TPair<FRigComponentKey, FRigBodyRecord>& BodyRecordPair : BodyRecords)
+		{
+			FRigBodyRecord& Record = BodyRecordPair.Value;
+			if (Record.ActorHandle)
+			{
+				Simulation->DestroyActor(Record.ActorHandle);
+			}
+		}
+
+		for (TPair<FRigComponentKey, FRigJointRecord>& JointRecordPair : JointRecords)
+		{
+			FRigJointRecord& Record = JointRecordPair.Value;
+			if (Record.JointHandle)
+			{
+				Simulation->DestroyJoint(Record.JointHandle);
+			}
+		}
+
+		for (TPair<FRigComponentKey, FRigControlRecord>& ControlRecordPair : ControlRecords)
+		{
+			FRigControlRecord& Record = ControlRecordPair.Value;
+			if (Record.JointHandle)
+			{
+				Simulation->DestroyJoint(Record.JointHandle);
+			}
+		}
+
+		if (CollisionActorHandle)
+		{
+			Simulation->DestroyActor(CollisionActorHandle);
+		}
+
+		if (SimulationActorHandle)
+		{
+			Simulation->DestroyActor(SimulationActorHandle);
+		}
+	}
+
+	BodyRecords.Empty();
+	JointRecords.Empty();
+	ControlRecords.Empty();
+
+	CollisionActorHandle = nullptr;
+	SimulationActorHandle = nullptr;
+
+	Simulation.Reset();
+}
+
+//======================================================================================================================
+static void SetCommonProperties(const FRigPhysicsCollisionShape& Shape, FKShapeElem& ShapeElem)
+{
+	ShapeElem.RestOffset = Shape.RestOffset;
+	ShapeElem.SetName(Shape.Name);
+	ShapeElem.SetContributeToMass(Shape.bContributeToMass);
+#ifdef PER_SHAPE_COLLISION
+	// Note that FKShapeElem supports enabling/disabling collision per shape, but this is discarded by the immediate
+	// solver.
+	ShapeElem.SetCollisionEnabled(Shape.CollisionEnabled);
+#endif
+}
+
+//======================================================================================================================
+static bool CreateGeometry(
+	const FRigPhysicsCollision&               Collision,
+	const FRigPhysicsDynamics*                Dynamics,
+	const Chaos::FReal                        Density,
+	Chaos::FReal&                             OutMass, 
+	Chaos::FVec3&                             OutInertia, 
+	Chaos::FRigidTransform3&                  OutCoMTransform, 
+	Chaos::FImplicitObjectPtr&                OutGeom, 
+	TArray<TUniquePtr<Chaos::FPerShapeData>>& OutShapes)
+{
+	using namespace Chaos;
+
+	OutMass = 0.0f;
+	OutInertia = FVector::ZeroVector;
+	OutCoMTransform = FTransform::Identity;
+
+	// Set the filter to collide with everything (we use a broad phase that only contains particle
+	// pairs that are explicitly set to collide)
+	FBodyCollisionData BodyCollisionData;
+
+	const Chaos::Filter::FShapeFilterData ShapeFilter = Chaos::Filter::FShapeFilterBuilder::BuildBlockAll(Chaos::EFilterFlags::All);
+	BodyCollisionData.CollisionFilterData.ComplexShapeFilterData = BodyCollisionData.CollisionFilterData.SimpleShapeFilterData = ShapeFilter;
+	BodyCollisionData.CollisionFilterData.FilterInstanceData = Chaos::Filter::FInstanceData();
+
+	// See FBodyInstance::BuildBodyCollisionFlags
+	BodyCollisionData.CollisionFlags.bEnableQueryCollision = false;
+	BodyCollisionData.CollisionFlags.bEnableSimCollisionSimple = true;
+	BodyCollisionData.CollisionFlags.bEnableSimCollisionComplex = false;
+	BodyCollisionData.CollisionFlags.bEnableProbeCollision = false;
+
+	FKAggregateGeom AggGeom;
+	for (const FRigPhysicsCollisionBox& Shape : Collision.Boxes)
+	{
+		FKBoxElem Elem(Shape.Extents.X, Shape.Extents.Y, Shape.Extents.Z);
+		SetCommonProperties(Shape, Elem);
+		Elem.Center = Shape.TM.GetTranslation();
+		Elem.Rotation = Shape.TM.Rotator();
+		AggGeom.BoxElems.Add(Elem);
+	}
+
+	for (const FRigPhysicsCollisionSphere& Shape : Collision.Spheres)
+	{
+		FKSphereElem Elem(Shape.Radius);
+		SetCommonProperties(Shape, Elem);
+		Elem.Center = Shape.TM.GetTranslation();
+		// Note that there is no rotation
+		AggGeom.SphereElems.Add(Elem);
+	}
+
+	for (const FRigPhysicsCollisionCapsule& Shape : Collision.Capsules)
+	{
+		FKSphylElem Elem(Shape.Radius, Shape.Length);
+		SetCommonProperties(Shape, Elem);
+		Elem.Center = Shape.TM.GetTranslation();
+		Elem.Rotation = Shape.TM.Rotator();
+		AggGeom.SphylElems.Add(Elem);
+	}
+
+	FGeometryAddParams AddParams;
+	AddParams.CollisionData = BodyCollisionData;
+	AddParams.CollisionTraceType = ECollisionTraceFlag::CTF_UseSimpleAsComplex;
+	AddParams.Scale = FVector(1, 1, 1);
+	AddParams.LocalTransform = FTransform::Identity; // How are these used? We will just set TM afterwards anyway
+	AddParams.WorldTransform = FTransform::Identity;
+	AddParams.Geometry = &AggGeom;
+
+	TArray<Chaos::FImplicitObjectPtr> Geoms;
+	FShapesArray Shapes;
+	ChaosInterface::CreateGeometry(AddParams, Geoms, Shapes);
+
+	if (Geoms.Num() == 0)
+	{
+		return false;
+	}
+
+	// Calculate mass properties, if we have dynamics
+	if (Dynamics)
+	{
+		// Whether each shape contributes to mass. It would be easier if ComputeMassProperties knew
+		// how to extract this info. Maybe it should be a flag in PerShapeData
+		TArray<bool> bContributesToMass;
+		bContributesToMass.Reserve(Shapes.Num());
+		for (int32 ShapeIndex = 0; ShapeIndex < Shapes.Num(); ++ShapeIndex)
+		{
+			const TUniquePtr<FPerShapeData>& Shape = Shapes[ShapeIndex];
+			const FKShapeElem* ShapeElem = FChaosUserData::Get<FKShapeElem>(Shape->GetUserData());
+			bool bHasMass = ShapeElem && ShapeElem->GetContributeToMass();
+			bContributesToMass.Add(bHasMass);
+		}
+
+		Chaos::FMassProperties MassProperties;
+		ChaosInterface::CalculateMassPropertiesFromShapeCollection(
+			MassProperties, Shapes, bContributesToMass, Density);
+
+		OutMass = MassProperties.Mass;
+		OutInertia = MassProperties.InertiaTensor.GetDiagonal();
+		OutCoMTransform = FTransform(MassProperties.RotationOfMass, MassProperties.CenterOfMass);
+	}
+
+	// If we have multiple root shapes, wrap them in a union
+	if (Geoms.Num() == 1)
+	{
+		OutGeom = MoveTemp(Geoms[0]);
+	}
+	else
+	{
+		OutGeom = MakeImplicitObjectPtr<FImplicitObjectUnion>(MoveTemp(Geoms));
+	}
+
+	for (TUniquePtr<FPerShapeData>& Shape : Shapes)
+	{
+		OutShapes.Emplace(MoveTemp(Shape));
+	}
+
+	return true;
+}
+
+//======================================================================================================================
+ImmediatePhysics::FActorHandle* FRigPhysicsSolver::CreateBody(
+	const FName                        BodyName,
+	const FRigPhysicsCollision&        Collision,
+	const FRigPhysicsDynamics*         Dynamics,
+	const FPhysicsControlModifierData* BodyData,
+	const FTransform&                  BodyRelSimSpaceTM) const
+{
+	ImmediatePhysics::FActorSetup ActorSetup;
+
+	if (Dynamics)
+	{
+		ActorSetup.ActorType = ImmediatePhysics::EActorType::DynamicActor;
+		ActorSetup.bEnableGravity = true;
+		ActorSetup.LinearDamping = Dynamics->LinearDamping;
+		ActorSetup.AngularDamping = Dynamics->AngularDamping;
+	}
+	else
+	{
+		ActorSetup.ActorType = ImmediatePhysics::EActorType::KinematicActor;
+	}
+
+	if (BodyData)
+	{
+		ActorSetup.bUpdateKinematicFromSimulation = BodyData->bUpdateKinematicFromSimulation;
+	}
+	else
+	{
+		ActorSetup.bUpdateKinematicFromSimulation = false;
+	}
+
+	Chaos::FVec3 Inertia;
+	Chaos::FRigidTransform3 CoMTransform;
+	Chaos::FReal Mass;
+	Chaos::FImplicitObjectPtr BodyGeom;
+	TArray<TUniquePtr<Chaos::FPerShapeData>> BodyShapes;
+	Chaos::FReal Density = Dynamics ? Dynamics->Density : 1.0f;
+	// Convert from g/cm^3 to kg/cm^3 - slightly odd units, but when we come to calculate density,
+	// distances will be in cm and masses are in kg. See CalculateMassPropertiesFromShapeCollection
+	Density *= 1e-3;
+
+	bool bGeometryCreated = CreateGeometry(
+		Collision, Dynamics, Density, Mass, Inertia, CoMTransform, BodyGeom, BodyShapes);
+
+	// We will have created with an arbitrary density - adjust to result in the desired mass.
+	ActorSetup.Mass = Mass;
+	ActorSetup.Inertia = Inertia;
+	ActorSetup.Transform = BodyRelSimSpaceTM;
+	ActorSetup.CoMTransform = CoMTransform;
+	ActorSetup.Geometry = MoveTemp(BodyGeom);
+	ActorSetup.Shapes = MoveTemp(BodyShapes);
+
+	if (Dynamics)
+	{
+		if (Mass > 0.0f && Dynamics->MassOverride > 0.0f)
+		{
+			ActorSetup.Mass = Dynamics->MassOverride;
+			ActorSetup.Inertia = (Dynamics->MassOverride / Mass) * Inertia;
+		}
+		if (Dynamics->bOverrideMomentsOfInertia)
+		{
+			ActorSetup.Inertia = Dynamics->MomentsOfInertiaOverride;
+		}
+		if (Dynamics->bOverrideCentreOfMass)
+		{
+			ActorSetup.CoMTransform.SetLocation(Dynamics->CentreOfMassOverride);
+		}
+		ActorSetup.CoMTransform.AddToTranslation(Dynamics->CentreOfMassNudge);
+	}
+
+
+	ActorSetup.Material = MakeUnique<Chaos::FChaosPhysicsMaterial>();
+	ActorSetup.Material->Friction = Collision.Material.Friction;
+	ActorSetup.Material->StaticFriction = Collision.Material.Friction;
+	ActorSetup.Material->Restitution = Collision.Material.Restitution;
+	ActorSetup.Material->FrictionCombineMode =
+		(Chaos::FChaosPhysicsMaterial::ECombineMode)Collision.Material.FrictionCombineMode;
+	ActorSetup.Material->RestitutionCombineMode =
+		(Chaos::FChaosPhysicsMaterial::ECombineMode)Collision.Material.RestitutionCombineMode;
+
+	ImmediatePhysics::FActorHandle* ActorHandle = Simulation->CreateActor(MoveTemp(ActorSetup));
+	if (!ActorHandle)
+	{
+		UE_LOG(LogRigPhysics, Warning,
+			TEXT("Unable to create body %s"), *BodyName.ToString());
+		return nullptr;
+	}
+
+	ActorHandle->SetName(BodyName);
+#if CHAOS_DEBUG_NAME
+	if (Chaos::FGeometryParticleHandle* ParticleHandle = ActorHandle->GetParticle())
+	{
+		ParticleHandle->SetDebugName(MakeShared<FString>(BodyName.ToString()));
+	}
+#endif
+
+	if (bGeometryCreated)
+	{
+		Simulation->AddToCollidingPairs(ActorHandle);
+		if (Dynamics)
+		{
+			// Note that particles are always created disabled. They will simulate when disabled,
+			// but won't collide!
+			ActorHandle->SetEnabled(true);
+		}
+	}
+	else
+	{
+		Simulation->SetHasCollision(ActorHandle, false);
+	}
+
+	return ActorHandle;
+}
+
+//======================================================================================================================
+void FRigPhysicsSolver::Instantiate(
+	const FRigVMExecuteContext&       ExecuteContext, 
+	URigHierarchy&                    Hierarchy, 
+	const FRigPhysicsSolverComponent& SolverComponent)
+{
+	DECLARE_SCOPE_HIERARCHICAL_COUNTER_FUNC()
+	QUICK_SCOPE_CYCLE_COUNTER(STAT_RigPhysics_Instantiate);
+
+	if (!bNeedToInstantiate)
+	{
+		return;
+	}
+
+	// Lock against access to the simulation, in case of access by the WorldObject task
+	UE::TScopeLock<FTransactionallySafeCriticalSection> Lock(*SimulationMutex.Get());
+
+	// We need the simulation space in order to instantiate properly. This is not ideal, as we may
+	// end up updating the simulation space data twice (thus inserting data into the history). This
+	// shouldn't really matter as it will only be on the first step.
+	UpdateSimulationSpaceStateAndCalculateData(
+		ExecuteContext, Hierarchy, SolverComponent, 0.0f, ExecuteContext.GetAbsoluteTime());
+
+	InitialiseSimulation();
+
+	InitialiseBodyRecords(Hierarchy, SolverComponent);
+
+	InitialiseJointRecords(Hierarchy, SolverComponent);
+
+	InitialiseControlRecords(Hierarchy, SolverComponent);
+
+	FRigPhysicsIgnorePairs IgnorePairs;
+
+	InstantiatePhysicsBodies(Hierarchy, SolverComponent, IgnorePairs);
+
+	InstantiatePhysicsJoints(Hierarchy, SolverComponent, IgnorePairs);
+
+	InstantiateControls(Hierarchy, SolverComponent, IgnorePairs);
+
+	// This is done last as it applies IgnorePairs
+	InstantiateSolverCollision(SolverComponent, IgnorePairs);
+
+	bNeedToInstantiate = false;
+}
+
+//======================================================================================================================
+ImmediatePhysics::FActorHandle* FRigPhysicsSolver::GetActor(const FRigComponentKey& ComponentKey) const
+{
+	if (const FRigBodyRecord* BodyRecord = BodyRecords.Find(ComponentKey))
+	{
+		return BodyRecord->ActorHandle;
+	}
+	// Also check to see if it is our own (the solver shapes)
+	if (ComponentKey == PhysicsSolverComponentKey)
+	{
+		return SimulationActorHandle;
+	}
+	return nullptr;
+}
+
+//======================================================================================================================
+void FRigPhysicsSolver::InstantiateSolverCollision(
+	const FRigPhysicsSolverComponent& SolverComponent,
+	FRigPhysicsIgnorePairs&           IgnorePairs)
+{
+	// Optionally create an object to contain environment collision
+	if (!SolverComponent.SolverSettings.Collision.IsEmpty())
+	{
+		// When we make these additional collision shapes, their actors are all considered to be at
+		// the origin, with the offsets being contained in the collision shapes.
+		FTransform BodyRelSimSpaceTM = ConvertCollisionSpaceTransformToSimSpace(
+			SolverComponent.SolverSettings, FTransform());
+
+		CollisionActorHandle = CreateBody(
+			FName(TEXT("Environment")), SolverComponent.SolverSettings.Collision, nullptr, nullptr, BodyRelSimSpaceTM);
+	}
+
+	// Add no-collision pairs
+	TArray<ImmediatePhysics_Chaos::FSimulation::FIgnorePair> ChaosIgnorePairs;
+
+	for (const FRigPhysicsIgnorePair& IgnorePair : IgnorePairs)
+	{
+		ImmediatePhysics::FSimulation::FIgnorePair ChaosPair;
+		ChaosPair.A = GetActor(IgnorePair.A);
+		ChaosPair.B = GetActor(IgnorePair.B);
+		if (ChaosPair.A && ChaosPair.B)
+		{
+			ChaosIgnorePairs.Add(ChaosPair);
+		}
+	}
+	Simulation->SetIgnoreCollisionPairTable(ChaosIgnorePairs);
+
+}
+
+//======================================================================================================================
+void FRigPhysicsSolver::InstantiatePhysicsBodies(
+	URigHierarchy&                    Hierarchy,
+	const FRigPhysicsSolverComponent& SolverComponent,
+	FRigPhysicsIgnorePairs&           IgnorePairs)
+{
+	DECLARE_SCOPE_HIERARCHICAL_COUNTER_FUNC()
+	QUICK_SCOPE_CYCLE_COUNTER(STAT_RigPhysics_InstantiatePhysicsBodies);
+
+	const FRigPhysicsSolverSettings& SolverSettings = SolverComponent.SolverSettings;
+	const FRigPhysicsSimulationSpaceSettings& SimulationSpaceSettings = SolverComponent.SimulationSpaceSettings;
+
+	for (TPair<FRigComponentKey, FRigBodyRecord>& BodyRecordPair : BodyRecords)
+	{
+		const FRigComponentKey& ComponentKey = BodyRecordPair.Key;
+		FRigBodyRecord& Record = BodyRecordPair.Value;
+
+		if (FRigPhysicsBodyComponent* PhysicsBodyComponent = Cast<FRigPhysicsBodyComponent>(
+			Hierarchy.FindComponent(ComponentKey)))
+		{
+			const FRigPhysicsCollision& Collision = PhysicsBodyComponent->Collision;
+			const FRigPhysicsDynamics& Dynamics = PhysicsBodyComponent->Dynamics;
+			const FPhysicsControlModifierData& BodyData = PhysicsBodyComponent->BodyData;
+
+			FRigElementKey SourceKey = PhysicsBodyComponent->BodySolverSettings.SourceBone;
+			if (!SourceKey.IsValid())
+			{
+				SourceKey = ComponentKey.ElementKey;
+			}
+
+			// What should we do if the key is not valid? 
+			if (SourceKey.IsValid())
+			{
+				Record.FinalComponentSpaceTM = Hierarchy.GetGlobalTransform(SourceKey);
+				const FTransform SourceSimulationSpaceTM =
+					ConvertComponentSpaceTransformToSimSpace(
+						SolverComponent.SolverSettings, Record.FinalComponentSpaceTM);
+				Record.ActorHandle = CreateBody(
+					ComponentKey.ElementKey.Name, Collision, &Dynamics, &BodyData, SourceSimulationSpaceTM);
+				if (Record.ActorHandle)
+				{
+					SetPhysicsBodyComponentState(SolverSettings, Record, *PhysicsBodyComponent);
+				}
+			}
+
+			Record.TargetElementKey = PhysicsBodyComponent->BodySolverSettings.TargetBone;
+			if (!Record.TargetElementKey.IsValid())
+			{
+				Record.TargetElementKey = ComponentKey.ElementKey;
+			}
+
+			for (const FRigComponentKey& NoCollisionKey : PhysicsBodyComponent->NoCollisionBodies)
+			{
+				IgnorePairs.Add(FRigPhysicsIgnorePair(ComponentKey, NoCollisionKey));
+			}
+		}
+	}
+}
+
+//======================================================================================================================
+void FRigPhysicsSolver::SetPhysicsBodyComponentState(
+	const FRigPhysicsSolverSettings& SolverSettings,
+	FRigBodyRecord&                  Record, 
+	FRigPhysicsBodyComponent&        PhysicsBodyComponent) const
+{
+	PhysicsBodyComponent.Transform = Record.FinalComponentSpaceTM;
+	PhysicsBodyComponent.CoMTransform = Record.ActorHandle->GetLocalCoMTransform() * Record.FinalComponentSpaceTM;
+	PhysicsBodyComponent.LinearVelocity = 
+		ConvertSimSpaceVectorToComponentSpace(SolverSettings, Record.ActorHandle->GetLinearVelocity());
+	PhysicsBodyComponent.AngularVelocity = 
+		ConvertSimSpaceVectorToComponentSpace(SolverSettings, Record.ActorHandle->GetAngularVelocity());
+}
+
+//======================================================================================================================
+void FRigPhysicsSolver::InstantiatePhysicsJoints(
+	const URigHierarchy&              Hierarchy,
+	const FRigPhysicsSolverComponent& SolverComponent,
+	FRigPhysicsIgnorePairs&           IgnorePairs)
+{
+	DECLARE_SCOPE_HIERARCHICAL_COUNTER_FUNC()
+	QUICK_SCOPE_CYCLE_COUNTER(STAT_RigPhysics_InstantiatePhysicsJoints);
+
+	const FRigPhysicsSolverSettings& SolverSettings = SolverComponent.SolverSettings;
+	const FRigPhysicsSimulationSpaceSettings& SimulationSpaceSettings = SolverComponent.SimulationSpaceSettings;
+
+	// Once all the bodies are created, we can make their physics joints
+	for (TPair<FRigComponentKey, FRigJointRecord>& JointRecordPair : JointRecords)
+	{
+		const FRigComponentKey& JointComponentKey = JointRecordPair.Key;
+		FRigJointRecord& JointRecord = JointRecordPair.Value;
+
+		if (const FRigPhysicsJointComponent* PhysicsJointComponent = Cast<FRigPhysicsJointComponent>(
+			Hierarchy.FindComponent(JointComponentKey)))
+		{
+			const FRigPhysicsJointData& JointData = PhysicsJointComponent->JointData;
+
+			FRigComponentKey ChildBodyComponentKey = JointRecord.ChildBodyComponentKey;
+			FRigElementKey ChildBoneKey = ChildBodyComponentKey.ElementKey;
+
+			FRigComponentKey ParentBodyComponentKey = JointRecord.ParentBodyComponentKey;
+			FRigElementKey ParentBoneKey = ParentBodyComponentKey.ElementKey;
+
+			// Joints require both parent and child to exist
+			const FRigPhysicsBodyComponent* ChildPhysicsBodyComponent =
+				Cast<FRigPhysicsBodyComponent>(Hierarchy.FindComponent(ChildBodyComponentKey));
+			FRigBodyRecord* ChildBodyRecord = BodyRecords.Find(ChildBodyComponentKey);
+			ImmediatePhysics::FActorHandle* ChildActorHandle = ChildBodyRecord ? ChildBodyRecord->ActorHandle : nullptr;
+			if (!ChildActorHandle)
+			{
+				continue;
+			}
+
+			const FRigPhysicsBodyComponent* ParentPhysicsBodyComponent =
+				Cast<FRigPhysicsBodyComponent>(Hierarchy.FindComponent(ParentBodyComponentKey));
+			FRigBodyRecord* ParentBodyRecord = BodyRecords.Find(ParentBodyComponentKey);
+			ImmediatePhysics::FActorHandle* ParentActorHandle = ParentBodyRecord ? ParentBodyRecord->ActorHandle : nullptr;
+			if (!ParentActorHandle)
+			{
+				continue;
+			}
+
+			// Make the physics joint (joint constraint). 
+			// UE likes to treat Body1 (index 0) as the child, and Body2 (index 1) as the parent
+			{
+				const FTransform ParentCoMTransform = ParentActorHandle->GetLocalCoMTransform();
+				const FTransform ChildCoMTransform = ChildActorHandle->GetLocalCoMTransform();
+
+				Chaos::FPBDJointSettings JointSettings;
+				if (JointData.bAutoCalculateChildOffset)
+				{
+					JointSettings.ConnectorTransforms[ConstraintChildIndex] = FTransform();
+				}
+				JointSettings.ConnectorTransforms[ConstraintChildIndex] =
+					JointData.ExtraChildOffset * JointSettings.ConnectorTransforms[ConstraintChildIndex];
+
+				if (JointData.bAutoCalculateParentOffset)
+				{
+					// When auto-calculating, the transform stored in the parent frame should place
+					// it at the child bone's location.
+					JointSettings.ConnectorTransforms[ConstraintParentIndex] = 
+						Hierarchy.GetLocalTransform(ChildBoneKey, true);
+				}
+				JointSettings.ConnectorTransforms[ConstraintParentIndex] =
+					JointData.ExtraParentOffset * JointSettings.ConnectorTransforms[ConstraintParentIndex];
+
+				ImmediatePhysics::UpdateJointSettingsFromLinearConstraint(JointData.LinearConstraint, JointSettings);
+				ImmediatePhysics::UpdateJointSettingsFromConeConstraint(JointData.ConeConstraint, JointSettings);
+				ImmediatePhysics::UpdateJointSettingsFromTwistConstraint(JointData.TwistConstraint, JointSettings);
+
+				// The physics setting is backwards, because we can't enable collision on bodies that
+				// are set to not collide for other reasons.
+				JointSettings.bCollisionEnabled = !JointData.bDisableCollision;
+				JointSettings.bProjectionEnabled = 
+					JointData.LinearProjectionAmount > 0.0f || JointData.AngularProjectionAmount > 0.0f;
+				JointSettings.AngularProjection = JointData.AngularProjectionAmount;
+				JointSettings.LinearProjection = JointData.LinearProjectionAmount;
+				JointSettings.ParentInvMassScale = JointData.ParentInverseMassScale;
+
+				JointSettings.bUseLinearSolver = SolverSettings.bUseLinearJointSolver;
+
+				JointRecord.JointHandle = Simulation->CreateJoint(
+					ImmediatePhysics::FJointSetup(JointSettings, ChildActorHandle, ParentActorHandle));
+
+				if (JointRecord.JointHandle)
+				{
+					if (Chaos::FPBDJointConstraintHandle* Constraint = JointRecord.JointHandle->GetConstraint())
+					{
+						Chaos::FPBDJointSettings Settings = Constraint->GetSettings();
+						if (!Settings.bCollisionEnabled)
+						{
+							IgnorePairs.Add(FRigPhysicsIgnorePair(
+								JointRecord.ChildBodyComponentKey, JointRecord.ParentBodyComponentKey));
+						}
+					}
+				}
+			}
+		}
+	}
+}
+
+//======================================================================================================================
+void FRigPhysicsSolver::InstantiateControls(
+	const URigHierarchy&              Hierarchy,
+	const FRigPhysicsSolverComponent& SolverComponent,
+	FRigPhysicsIgnorePairs&           IgnorePairs)
+{
+	DECLARE_SCOPE_HIERARCHICAL_COUNTER_FUNC()
+	QUICK_SCOPE_CYCLE_COUNTER(STAT_RigPhysics_InstantiateControls);
+
+	if (!Simulation.IsValid())
+	{
+		return;
+	}
+
+	for (TPair<FRigComponentKey, FRigControlRecord>& ControlRecordPair : ControlRecords)
+	{
+		const FRigComponentKey& ComponentKey = ControlRecordPair.Key;
+		FRigControlRecord& ControlRecord = ControlRecordPair.Value;
+
+		if (const FRigPhysicsControlComponent* ControlComponent = Cast<FRigPhysicsControlComponent>(
+			Hierarchy.FindComponent(ComponentKey)))
+		{
+			if (FRigBodyRecord* ChildBodyRecord = BodyRecords.Find(ControlRecord.ChildBodyComponentKey))
+			{
+				FRigBodyRecord* ParentBodyRecord = BodyRecords.Find(ControlRecord.ParentBodyComponentKey);
+
+				// ParentBodyRecord being nullptr just means to create a sim-space control. However,
+				// we never want to make a sim-space control if the user has set
+				// bUseParentBodyAsDefault, because that flag implies the user wanted a parent-space
+				// control. Note that when creating chains etc it's quite easy to "accidentally"
+				// create a parent-space control for the first body in the chain, even when one
+				// isn't actually required. In fact, this can happen very easily with modular rigs
+				// because chains may be created independently, and we'll only find out we want a
+				// parent-space control on the arm (say) after the spine has been created.
+				if (!ParentBodyRecord && ControlComponent->bUseParentBodyAsDefault)
+				{
+					// We will have a control in the component, but it will do nothing as there will
+					// be no physics joint/constraint.
+					continue;
+				}
+
+				ImmediatePhysics::FActorHandle* const ChildBodyHandle = ChildBodyRecord->ActorHandle;
+				ImmediatePhysics::FActorHandle* const ParentBodyHandle = 
+					ParentBodyRecord ? ParentBodyRecord->ActorHandle : SimulationActorHandle;
+
+				if (ChildBodyHandle && ParentBodyHandle)
+				{
+					// The constraint is created disabled - it will be updated in pre-physics
+					ControlRecord.JointHandle = CreatePhysicsJoint(*Simulation.Get(), *ChildBodyHandle, *ParentBodyHandle);
+				}
+
+				if (!ControlRecord.JointHandle)
+				{
+					UE_LOG(LogRigPhysics, Warning,
+						TEXT("Unable to create control constraint for %s"), 
+						*ControlComponent->GetKey().ToString());
+				}
+
+				if (ControlComponent->ControlData.bDisableCollision)
+				{
+					IgnorePairs.Add(FRigPhysicsIgnorePair(
+						ControlRecord.ChildBodyComponentKey, ControlRecord.ParentBodyComponentKey));
+				}
+			}
+		}
+	}
+}
+
+//======================================================================================================================
+void FRigPhysicsSolver::UpdateBodyRecordPrePhysics(
+	const URigHierarchy&              Hierarchy,
+	const FRigPhysicsSolverComponent& SolverComponent,
+	const float                       DeltaTime,
+	FRigBodyRecord&                   Record, 
+	const FRigPhysicsBodyComponent*   PhysicsBodyComponent)
+{
+	if (!Record.ActorHandle)
+	{
+		return;
+	}
+	const FRigComponentKey ComponentKey = PhysicsBodyComponent->GetKey();
+
+	// Shuffle the record data
+	Record.PrevSourceComponentSpaceVelocity = Record.SourceComponentSpaceVelocity;
+	Record.PrevSourceComponentSpaceAngularVelocity = Record.SourceComponentSpaceAngularVelocity;
+	Record.PrevSourceComponentSpaceTM = Record.SourceComponentSpaceTM;
+
+	FRigElementKey SourceKey = PhysicsBodyComponent->BodySolverSettings.SourceBone;
+	if (!SourceKey.IsValid())
+	{
+		SourceKey = ComponentKey.ElementKey;
+	}
+	if (SourceKey.IsValid())
+	{
+		Record.SourceComponentSpaceTM = Hierarchy.GetGlobalTransform(SourceKey);
+	}
+
+	if (UpdateCounter == PreviousUpdateCounter + 1)
+	{
+		Record.SourceComponentSpaceVelocity = UE::PhysicsControl::CalculateLinearVelocity(
+			Record.PrevSourceComponentSpaceTM.GetTranslation(), Record.SourceComponentSpaceTM.GetTranslation(), DeltaTime);
+		Record.SourceComponentSpaceAngularVelocity = UE::PhysicsControl::CalculateAngularVelocity(
+			Record.PrevSourceComponentSpaceTM.GetRotation(), Record.SourceComponentSpaceTM.GetRotation(), DeltaTime);
+	}
+	else
+	{
+		Record.PrevSourceComponentSpaceTM = Record.SourceComponentSpaceTM;
+		Record.SourceComponentSpaceVelocity = FVector::ZeroVector;
+		Record.SourceComponentSpaceAngularVelocity = FVector::ZeroVector;
+	}
+}
+
+//======================================================================================================================
+void FRigPhysicsSolver::UpdateBodyPrePhysics(
+	const FRigVMExecuteContext&       ExecuteContext, 
+	const FRigPhysicsSolverComponent& SolverComponent,
+	const FRigBodyRecord&             Record, 
+	const FRigPhysicsBodyComponent&   PhysicsBodyComponent)
+{
+	if (!Record.ActorHandle)
+	{
+		return;
+	}
+	const FRigComponentKey ComponentKey = PhysicsBodyComponent.GetKey();
+
+	FPhysicsControlModifierData BodyData = PhysicsBodyComponent.BodyData;
+	EPhysicsControlKinematicTargetSpace KinematicTargetSpace = BodyData.KinematicTargetSpace;
+
+	if (SolverComponent.TrackInputCounter > 0)
+	{
+		// If the body is already kinematic, then we don't want to make it start ignoring a target.
+		// However, dynamic objects should avoid using a target if they are just tracking due to a reset.
+		if (BodyData.MovementType != EPhysicsMovementType::Kinematic)
+		{
+			BodyData.MovementType = EPhysicsMovementType::Kinematic;
+			KinematicTargetSpace = EPhysicsControlKinematicTargetSpace::IgnoreTarget;
+		}
+	}
+
+	UpdateBodyFromModifierData(
+		*Record.ActorHandle, *Simulation.Get(), BodyData, SimulationSpaceData.Gravity);
+
+	if (Record.ActorHandle->GetIsKinematic())
+	{
+		// Get the target in component space, and then convert it into sim space if necessary.
+
+		// If the target is already in component space, then that's all we need
+		FTransform KinematicTargetCS;
+		const FTransform KinematicTargetTM = PhysicsBodyComponent.KinematicTarget;
+		switch (KinematicTargetSpace)
+		{
+		case EPhysicsControlKinematicTargetSpace::Component:
+		{
+			KinematicTargetCS = KinematicTargetTM;
+			break;
+		}
+		case EPhysicsControlKinematicTargetSpace::World:
+		{
+			KinematicTargetCS = KinematicTargetTM.GetRelativeTransform(SimulationSpaceState.ComponentTM);
+			break;
+		}
+		default:
+		{
+			// All the other options are relative to a bone, so the first task is to get
+			// that, which will be in component space
+			FRigElementKey SourceKey = PhysicsBodyComponent.BodySolverSettings.SourceBone;
+			if (!SourceKey.IsValid())
+			{
+				SourceKey = ComponentKey.ElementKey;
+			}
+			if (SourceKey.IsValid())
+			{
+				switch (KinematicTargetSpace)
+				{
+				case EPhysicsControlKinematicTargetSpace::OffsetInBoneSpace:
+				{
+					KinematicTargetCS = KinematicTargetTM * Record.SourceComponentSpaceTM;
+					break;
+				}
+				case EPhysicsControlKinematicTargetSpace::OffsetInWorldSpace:
+				{
+					const FTransform ComponentTM(SimulationSpaceState.ComponentTM);
+					// Convert the bone to WS (it's really component space!), apply the target, and convert back. 
+					FTransform KinematicTargetWS = Record.SourceComponentSpaceTM * ComponentTM;
+					// Note that we don't want to rotate the translation by the target - we apply it
+					// individually to orientation and position.
+					KinematicTargetWS.AddToTranslation(KinematicTargetTM.GetTranslation());
+					KinematicTargetWS.SetRotation(KinematicTargetTM.GetRotation() * KinematicTargetWS.GetRotation());
+					KinematicTargetCS = KinematicTargetWS.GetRelativeTransform(ComponentTM);
+					break;
+				}
+				case EPhysicsControlKinematicTargetSpace::OffsetInComponentSpace:
+				{
+					// This applies the offset as a transform in component space. Note that this is
+					// different to what PCC does, because in ControlRig "world" is really
+					// "component" space.
+					KinematicTargetCS =  Record.SourceComponentSpaceTM * KinematicTargetTM;
+					break;
+				}
+				case EPhysicsControlKinematicTargetSpace::IgnoreTarget:
+				{
+					KinematicTargetCS = Record.SourceComponentSpaceTM;
+					break;
+				}
+				default:
+				{
+					KinematicTargetCS = Record.SourceComponentSpaceTM;
+					UE_CONTROLRIG_RIGUNIT_REPORT_ERROR(TEXT("Kinematic target space is not valid"));
+				}
+				}
+			}
+		}
+		}
+		const FTransform KinematicTargetSimSpaceTM = ConvertComponentSpaceTransformToSimSpace(
+			SolverComponent.SolverSettings, KinematicTargetCS);
+		Record.ActorHandle->SetKinematicTarget(KinematicTargetSimSpaceTM);
+	}
+	else
+	{
+		// Danny TODO move damping into BodyData - any PhysicsControl system should be able to use it
+		Record.ActorHandle->SetLinearDamping(PhysicsBodyComponent.Dynamics.LinearDamping);
+		Record.ActorHandle->SetAngularDamping(PhysicsBodyComponent.Dynamics.AngularDamping);
+	}
+}
+
+//======================================================================================================================
+void FRigPhysicsSolver::ApplyForcesPrePhysics(
+	const FRigVMExecuteContext&       ExecuteContext,
+	const FRigPhysicsSolverComponent& SolverComponent,
+	const FRigBodyRecord&             Record,
+	const FRigPhysicsBodyComponent&   PhysicsBodyComponent)
+{
+	if (PhysicsBodyComponent.ForceAndTorques.IsEmpty() || !Record.ActorHandle || Record.ActorHandle->GetIsKinematic())
+	{
+		return;
+	}
+
+	ImmediatePhysics::FActorHandle& ActorHandle = *Record.ActorHandle;
+	const FRigPhysicsSolverSettings& SolverSettings = SolverComponent.SolverSettings;
+
+	for (const FPhysicsControlNamedForceAndTorqueData& Data : PhysicsBodyComponent.ForceAndTorques)
+	{
+		const FPhysicsControlForceAndTorqueData& ForceAndTorqueData = Data.ForceAndTorqueData;
+
+		// Need to convert the requested force/torque into simulation space first
+		FVector LinearSS;
+		FVector AngularSS;
+
+		FTransform BodyTM = ActorHandle.GetWorldTransform();
+
+		switch (ForceAndTorqueData.ForceAndTorqueSpace)
+		{
+		case EPhysicsControlSpace::World: // For Control Rig, World and Component space are considered to be the same
+		case EPhysicsControlSpace::Component:
+			LinearSS = ConvertComponentSpaceVectorToSimSpace(
+				SolverSettings, SimulationSpaceState.ComponentTM.InverseTransformVector(ForceAndTorqueData.Linear));
+			AngularSS = ConvertComponentSpaceVectorToSimSpace(
+				SolverSettings, SimulationSpaceState.ComponentTM.InverseTransformVector(ForceAndTorqueData.Angular));
+			break;
+		case EPhysicsControlSpace::Body:
+			LinearSS = BodyTM.GetRotation() * ForceAndTorqueData.Linear;
+			AngularSS = BodyTM.GetRotation() * ForceAndTorqueData.Angular;
+			break;
+		default:
+			LinearSS = FVector::ZeroVector;
+			AngularSS = FVector::ZeroVector;
+			break;
+		}
+
+		// Apply it differently depending on the location
+		ImmediatePhysics_Shared::EForceType FT = ConvertForceType(ForceAndTorqueData.Type);
+		switch (ForceAndTorqueData.LocationSpace)
+		{
+		case EPhysicsControlSpace::World: // For Control Rig, World and Component space are considered to be the same
+		case EPhysicsControlSpace::Component:
+		{
+			FVector SimSpaceLocation = ConvertWorldPositionToSimSpaceNoScale(
+				SolverSettings, ForceAndTorqueData.Location);
+			ActorHandle.AddForceAtLocation(LinearSS, SimSpaceLocation, FT);
+		}
+		break;
+		case EPhysicsControlSpace::Body:
+		{
+			FVector SimSpaceLocation = ActorHandle.GetWorldTransform().TransformPositionNoScale(
+				ActorHandle.GetLocalCoMLocation() + ForceAndTorqueData.Location);
+			ActorHandle.AddForceAtLocation(LinearSS, SimSpaceLocation, FT);
+		}
+		break;
+		default:
+			break;
+		}
+		// Torque doesn't get applied at a location
+		ActorHandle.AddTorque(AngularSS, FT);
+	}
+}
+
+//======================================================================================================================
+void FRigPhysicsSolver::UpdateJointPrePhysics(
+	const URigHierarchy&                Hierarchy, 
+	FRigJointRecord&                    Record, 
+	const FRigPhysicsJointComponent&    PhysicsJointComponent, 
+	const float                         DeltaTime)
+{
+	// Now update the joint targets
+	if (!Record.JointHandle)
+	{
+		return;
+	}
+	if (Chaos::FPBDJointConstraintHandle* Constraint = Record.JointHandle->GetConstraint())
+	{
+		const FRigComponentKey ComponentKey = PhysicsJointComponent.GetKey();
+
+		// Set the drive strength etc
+		const FRigPhysicsJointData& JointData = PhysicsJointComponent.JointData;
+		const FRigPhysicsDriveData& DriveData = PhysicsJointComponent.DriveData;
+
+		if (PhysicsJointComponent.DirtyFlags != ERigPhysicsJointComponentDirtyFlags::None)
+		{
+			Chaos::FPBDJointSettings Settings = Constraint->GetSettings();
+
+			if ((PhysicsJointComponent.DirtyFlags & ERigPhysicsJointComponentDirtyFlags::Joint) != 
+				ERigPhysicsJointComponentDirtyFlags::None)
+			{
+				SetPhysicsJointEnabled(*Record.JointHandle, JointData.bEnabled);
+
+				ImmediatePhysics::UpdateJointSettingsFromLinearConstraint(JointData.LinearConstraint, Settings);
+				ImmediatePhysics::UpdateJointSettingsFromConeConstraint(JointData.ConeConstraint, Settings);
+				ImmediatePhysics::UpdateJointSettingsFromTwistConstraint(JointData.TwistConstraint, Settings);
+			}
+
+			if ((PhysicsJointComponent.DirtyFlags & ERigPhysicsJointComponentDirtyFlags::Drive) !=
+				ERigPhysicsJointComponentDirtyFlags::None)
+			{
+				ImmediatePhysics::UpdateJointSettingsFromLinearDriveConstraint(DriveData.LinearDriveConstraint, Settings);
+				ImmediatePhysics::UpdateJointSettingsFromAngularDriveConstraint(DriveData.AngularDriveConstraint, Settings);
+			}
+
+			Constraint->SetSettings(Settings);
+
+			PhysicsJointComponent.DirtyFlags = ERigPhysicsJointComponentDirtyFlags::None;
+		}
+
+		const Chaos::FPBDJointSettings& Settings = Constraint->GetSettings();
+
+		// Now set the actual target
+		if (JointData.bEnabled && (
+			Settings.AngularDriveStiffness.SquaredLength() > 0.0f ||
+			Settings.AngularDriveDamping.SquaredLength() > 0.0f ||
+			Settings.LinearDriveStiffness.SquaredLength() > 0.0f ||
+			Settings.LinearDriveDamping.SquaredLength() > 0.0f))
+		{
+			// Multiplier on the velocity calculated from the current and previous target
+			//
+			// Danny TODO surely this should always be taken from DriveData? Do we/can we
+			// distinguish between manual and animation velocities?
+			float TargetVelocityMultiplier = 1.0f;
+
+			FRigElementKey ChildSourceKey;
+			if (const FRigPhysicsBodyComponent* ChildPhysicsBodyComponent = Cast<FRigPhysicsBodyComponent>(
+				Hierarchy.FindComponent(Record.ChildBodyComponentKey)))
+			{
+				ChildSourceKey = ChildPhysicsBodyComponent->BodySolverSettings.SourceBone;
+				if (!ChildSourceKey.IsValid())
+				{
+					ChildSourceKey = ComponentKey.ElementKey;
+				}
+				TargetVelocityMultiplier = DriveData.SkeletalAnimationVelocityMultiplier;
+			}
+
+			FTransform DriveTargetTM;
+			if (DriveData.bUseSkeletalAnimation)
+			{
+				FRigElementKey ParentSourceKey;
+				if (const FRigPhysicsBodyComponent* ParentPhysicsBodyComponent = Cast<FRigPhysicsBodyComponent>(
+					Hierarchy.FindComponent(Record.ParentBodyComponentKey)))
+				{
+					ParentSourceKey = ParentPhysicsBodyComponent->BodySolverSettings.SourceBone;
+					if (!ParentSourceKey.IsValid())
+					{
+						ParentSourceKey = Record.ParentBodyComponentKey.ElementKey;
+					}
+				}
+
+				if (!ChildSourceKey.IsValid() || !ParentSourceKey.IsValid())
+				{
+					return;
+				}
+
+				// Danny TODO now all transforms are being cached in the record, get them from there
+				// rather than the hierarchy
+
+				// Note that the drive operates between a parent and child part, so we
+				// don't need to worry about global/component (etc) space.
+				const FTransform ChildTM(Hierarchy.GetGlobalTransform(ChildSourceKey));
+				const FTransform ParentTM(Hierarchy.GetGlobalTransform(ParentSourceKey));
+
+				const FTransform ComponentSpaceParentFrameTM = 
+					Settings.ConnectorTransforms[ConstraintParentIndex] * ParentTM;
+				const FTransform ComponentSpaceChildFrameTM = 
+					Settings.ConnectorTransforms[ConstraintChildIndex] * ChildTM;
+
+				DriveTargetTM = ComponentSpaceChildFrameTM.GetRelativeTransform(ComponentSpaceParentFrameTM);
+			}
+
+			// Apply the offset in the frame of the (potentially animation) target
+			DriveTargetTM = FTransform(				
+				DriveData.AngularDriveConstraint.OrientationTarget.Quaternion(),
+				DriveData.LinearDriveConstraint.PositionTarget) * DriveTargetTM;
+
+			Constraint->SetLinearDrivePositionTarget(DriveTargetTM.GetTranslation());
+			Constraint->SetAngularDrivePositionTarget(DriveTargetTM.GetRotation());
+
+			if (Record.PreviousDriveTargetUpdateCounter + 1 == UpdateCounter && TargetVelocityMultiplier > 0.0f)
+			{
+				if (DeltaTime > SMALL_NUMBER)
+				{
+					const FTransform DriveTargetTMDelta =
+						Record.PreviousDriveTargetTM.GetRelativeTransformReverse(DriveTargetTM);
+					const FVector Velocity = DriveTargetTMDelta.GetTranslation() / DeltaTime;
+					const FVector AngularVelocity =
+						DriveTargetTMDelta.GetRotation().GetShortestArcWith(
+							FQuat::Identity).ToRotationVector() / DeltaTime;
+					if (!Velocity.ContainsNaN() && !AngularVelocity.ContainsNaN())
+					{
+						Constraint->SetLinearDriveVelocityTarget(Velocity * TargetVelocityMultiplier);
+						Constraint->SetAngularDriveVelocityTarget(AngularVelocity * TargetVelocityMultiplier);
+					}
+				}
+			}
+			else
+			{
+				Constraint->SetLinearDriveVelocityTarget(Chaos::FVec3(0));
+				Constraint->SetAngularDriveVelocityTarget(Chaos::FVec3(0));
+			}
+			Record.PreviousDriveTargetUpdateCounter = UpdateCounter;
+			Record.PreviousDriveTargetTM = DriveTargetTM;
+		}
+	}
+}
+
+//======================================================================================================================
+static FTransform CalculateTargetTM(
+	const URigHierarchy&               Hierarchy,
+	const Chaos::FPBDJointSettings&    JointSettings,
+	const FRigControlRecord&           Record)
+{
+	const FTransform ChildTM = Record.ChildBodyComponentKey.IsValid() ?
+		Hierarchy.GetGlobalTransform(Record.ChildBodyComponentKey.ElementKey) : FTransform();
+
+	const FTransform ChildTargetTM = JointSettings.ConnectorTransforms[ConstraintChildIndex] * ChildTM;
+
+	if (Record.ParentBodyComponentKey.IsValid())
+	{
+		const FTransform ParentTM(
+			Hierarchy.GetGlobalTransform(Record.ParentBodyComponentKey.ElementKey));
+		const FTransform ParentTargetTM = JointSettings.ConnectorTransforms[ConstraintParentIndex] * ParentTM;
+		return ChildTargetTM.GetRelativeTransform(ParentTargetTM);
+	}
+	return ChildTargetTM;
+}
+
+//======================================================================================================================
+void FRigPhysicsSolver::CheckForResetsPrePhysics(
+	const URigHierarchy& Hierarchy, FRigPhysicsSolverComponent& SolverComponent, const float DeltaTime)
+{
+	DECLARE_SCOPE_HIERARCHICAL_COUNTER_FUNC()
+	QUICK_SCOPE_CYCLE_COUNTER(STAT_RigPhysics_CheckForResetsPrePhysics);
+
+	const FRigPhysicsSolverSettings& SolverSettings = SolverComponent.SolverSettings;
+
+	double SpeedThresholdForResetSquared = SolverSettings.KinematicSpeedThresholdForReset > 0 ?
+		FMath::Square(SolverSettings.KinematicSpeedThresholdForReset) : TNumericLimits<double>::Max();
+
+	double AccelerationThresholdForResetSquared = SolverSettings.KinematicAccelerationThresholdForReset > 0 ?
+		FMath::Square(SolverSettings.KinematicAccelerationThresholdForReset) : TNumericLimits<double>::Max();
+
+	if (SpeedThresholdForResetSquared == TNumericLimits<double>::Max() &&
+		AccelerationThresholdForResetSquared == TNumericLimits<double>::Max())
+	{
+		return;
+	}
+
+	double HighestSpeedSq = -1.0;
+	double HighestAccelerationSq = -1.0;
+
+	for (TPair<FRigComponentKey, FRigBodyRecord>& BodyRecordPair : BodyRecords)
+	{
+		const FRigComponentKey& ComponentKey = BodyRecordPair.Key;
+		FRigBodyRecord& Record = BodyRecordPair.Value;
+
+		if (Record.ActorHandle && Record.ActorHandle->GetIsKinematic())
+		{
+			if (const FRigPhysicsBodyComponent* PhysicsBodyComponent = Cast<FRigPhysicsBodyComponent>(
+				Hierarchy.FindComponent(ComponentKey)))
+			{
+				if (PhysicsBodyComponent->BodySolverSettings.bIncludeInChecksForReset)
+				{
+					const FVector Velocity = Record.SourceComponentSpaceVelocity;
+					const FVector Acceleration =
+						DeltaTime > SMALL_NUMBER && (UpdateCounter == PreviousUpdateCounter + 1)
+						? (Record.SourceComponentSpaceVelocity - Record.PrevSourceComponentSpaceVelocity) / DeltaTime
+						: FVector::ZeroVector;
+
+					HighestSpeedSq = FMath::Max(HighestSpeedSq, Velocity.SquaredLength());
+					HighestAccelerationSq = FMath::Max(HighestAccelerationSq, Acceleration.SquaredLength());
+				}
+			}
+		}
+	}
+	if (HighestSpeedSq > SpeedThresholdForResetSquared || HighestAccelerationSq > AccelerationThresholdForResetSquared)
+	{
+		if (HighestSpeedSq > SpeedThresholdForResetSquared)
+		{
+			UE_LOG(LogRigPhysics, Log, TEXT("Speed %f triggered reset"),
+				FMath::Sqrt(HighestSpeedSq));
+		}
+		if (HighestAccelerationSq > AccelerationThresholdForResetSquared)
+		{
+			UE_LOG(LogRigPhysics, Log, TEXT("Acceleration %f triggered reset"),
+				FMath::Sqrt(HighestAccelerationSq));
+		}
+		SolverComponent.TrackInputCounter = FMath::Max(SolverComponent.TrackInputCounter, 3);
+	}
+}
+
+//======================================================================================================================
+void FRigPhysicsSolver::UpdatePrePhysics(
+	const FRigVMExecuteContext& ExecuteContext, 
+	const URigHierarchy&        Hierarchy,
+	FRigPhysicsSolverComponent& SolverComponent, 
+	const float                 DeltaTime)
+{
+	DECLARE_SCOPE_HIERARCHICAL_COUNTER_FUNC()
+	QUICK_SCOPE_CYCLE_COUNTER(STAT_RigPhysics_UpdatePrePhysics);
+
+	if (!Simulation.IsValid())
+	{
+		return;
+	}
+
+	if (CollisionActorHandle)
+	{
+		FTransform BodyRelSimSpaceTM = ConvertCollisionSpaceTransformToSimSpace(
+			SolverComponent.SolverSettings, FTransform());
+		CollisionActorHandle->SetKinematicTarget(BodyRelSimSpaceTM);
+	}
+
+	{
+		DECLARE_SCOPE_HIERARCHICAL_COUNTER_FUNC()
+		QUICK_SCOPE_CYCLE_COUNTER(STAT_RigPhysics_UpdateBodyRecordsPrePhysics);
+		for (TPair<FRigComponentKey, FRigBodyRecord>& BodyRecordPair : BodyRecords)
+		{
+			const FRigComponentKey& ComponentKey = BodyRecordPair.Key;
+			FRigBodyRecord& Record = BodyRecordPair.Value;
+			if (const FRigPhysicsBodyComponent* PhysicsBodyComponent = Cast<FRigPhysicsBodyComponent>(
+				Hierarchy.FindComponent(ComponentKey)))
+			{
+				UpdateBodyRecordPrePhysics(Hierarchy, SolverComponent, DeltaTime, Record, PhysicsBodyComponent);
+			}
+		}
+	}
+
+	CheckForResetsPrePhysics(Hierarchy, SolverComponent, DeltaTime);
+
+	if (SolverComponent.TrackInputCounter > 0)
+	{
+		UE_LOG(LogRigPhysics, Log, TEXT("Forcing tracking (counter = %d) of input"), 
+			SolverComponent.TrackInputCounter);
+	}
+
+	{
+		DECLARE_SCOPE_HIERARCHICAL_COUNTER_FUNC()
+		QUICK_SCOPE_CYCLE_COUNTER(STAT_RigPhysics_UpdateBodiesPrePhysics);
+		for (TPair<FRigComponentKey, FRigBodyRecord>& BodyRecordPair : BodyRecords)
+		{
+			const FRigComponentKey& ComponentKey = BodyRecordPair.Key;
+			FRigBodyRecord& Record = BodyRecordPair.Value;
+			if (const FRigPhysicsBodyComponent* PhysicsBodyComponent = Cast<FRigPhysicsBodyComponent>(
+				Hierarchy.FindComponent(ComponentKey)))
+			{
+				UpdateBodyPrePhysics(ExecuteContext, SolverComponent, Record, *PhysicsBodyComponent);
+			}
+		}
+	}
+
+	{
+		DECLARE_SCOPE_HIERARCHICAL_COUNTER_FUNC()
+		QUICK_SCOPE_CYCLE_COUNTER(STAT_RigPhysics_UpdateJointsPrePhysics);
+		for (TPair<FRigComponentKey, FRigJointRecord>& JointRecordPair : JointRecords)
+		{
+			const FRigComponentKey& ComponentKey = JointRecordPair.Key;
+			FRigJointRecord& JointRecord = JointRecordPair.Value;
+
+			if (const FRigPhysicsJointComponent* JointComponent = Cast<FRigPhysicsJointComponent>(
+				Hierarchy.FindComponent(ComponentKey)))
+			{
+				UpdateJointPrePhysics(Hierarchy, JointRecord, *JointComponent, DeltaTime);
+			}
+		}
+	}
+
+	{
+		DECLARE_SCOPE_HIERARCHICAL_COUNTER_FUNC()
+		QUICK_SCOPE_CYCLE_COUNTER(STAT_RigPhysics_UpdateControlsPrePhysics);
+		for (TPair<FRigComponentKey, FRigControlRecord>& ControlRecordPair : ControlRecords)
+		{
+			const FRigComponentKey& ComponentKey = ControlRecordPair.Key;
+			FRigControlRecord& ControlRecord = ControlRecordPair.Value;
+
+			if (const FRigPhysicsControlComponent* ControlComponent = Cast<FRigPhysicsControlComponent>(
+				Hierarchy.FindComponent(ComponentKey)))
+			{
+				UpdateControlPrePhysics(ControlRecord, *ControlComponent, SolverComponent, Hierarchy, DeltaTime);
+			}
+		}
+	}
+
+	{
+		DECLARE_SCOPE_HIERARCHICAL_COUNTER_FUNC()
+		QUICK_SCOPE_CYCLE_COUNTER(STAT_RigPhysics_ApplyForcesPrePhysics);
+		for (TPair<FRigComponentKey, FRigBodyRecord>& BodyRecordPair : BodyRecords)
+		{
+			const FRigComponentKey& ComponentKey = BodyRecordPair.Key;
+			FRigBodyRecord& Record = BodyRecordPair.Value;
+
+			if (const FRigPhysicsBodyComponent* PhysicsBodyComponent = Cast<FRigPhysicsBodyComponent>(
+				Hierarchy.FindComponent(ComponentKey)))
+			{
+				ApplyForcesPrePhysics(ExecuteContext, SolverComponent, Record, *PhysicsBodyComponent);
+			}
+
+		}
+	}
+}
+
+//======================================================================================================================
+void FRigPhysicsSolver::UpdateControlPrePhysics(
+	FRigControlRecord&                 ControlRecord,
+	const FRigPhysicsControlComponent& ControlComponent,
+	const FRigPhysicsSolverComponent&  SolverComponent,
+	const URigHierarchy&               Hierarchy,
+	const float                        DeltaTime)
+{
+	const FRigPhysicsSolverSettings& SolverSettings = SolverComponent.SolverSettings;
+
+	ImmediatePhysics::FJointHandle* JointHandle = ControlRecord.JointHandle;
+	if (!JointHandle)
+	{
+		return;
+	}
+	Chaos::FPBDJointConstraintHandle* Constraint = JointHandle->GetConstraint();
+	if (!Constraint)
+	{
+		return;
+	}
+
+	if (!ControlComponent.ControlData.bEnabled)
+	{
+		SetPhysicsJointEnabled(*JointHandle, false);
+		return;
+	}
+
+	float ThisDeltaTime = DeltaTime;
+
+	if (ControlRecord.PreviousTargetUpdateCounter + 1 != UpdateCounter)
+	{
+		// If we missed some intermediate updates, then we don't want to use the previous
+		// positions etc to calculate velocities. This will mean velocity/damping will be
+		// incorrect for one frame, but that's probably OK.
+		ThisDeltaTime = 0.0f;
+	}
+
+	if ((ControlComponent.DirtyFlags & ERigPhysicsControlComponentDirtyFlags::Data) !=
+		ERigPhysicsControlComponentDirtyFlags::None)
+	{
+		Constraint->SetCollisionEnabled(!ControlComponent.ControlData.bDisableCollision);
+		Constraint->SetParentInvMassScale(ControlComponent.ControlData.bOnlyControlChildObject ? 0 : 1);
+	}
+
+	// Check for early outs. Make sure that we don't apply velocities using the wrong calculation
+	// when the strength/damping is increased in the future by not updating the update counter.
+
+	const ImmediatePhysics::FActorHandle* ChildActorHandle = JointHandle->GetActorHandles()[ConstraintChildIndex];
+	const ImmediatePhysics::FActorHandle* ParentActorHandle = JointHandle->GetActorHandles()[ConstraintParentIndex];
+	if (!ChildActorHandle || !ParentActorHandle)
+	{
+		return;
+	}
+
+	const Chaos::FPBDJointSettings& JointSettings = Constraint->GetSettings();
+
+	// Note that if we don't have any strength, then we don't calculate the targets. 
+	if ((ControlComponent.DirtyFlags & ERigPhysicsControlComponentDirtyFlags::Data) !=
+		ERigPhysicsControlComponentDirtyFlags::None)
+	{
+		ControlComponent.DirtyFlags &= ~ERigPhysicsControlComponentDirtyFlags::Data;
+		if (!UpdateDriveSpringDamperSettings(
+			*JointHandle, JointSettings, ControlComponent.ControlData, ControlComponent.ControlMultiplier))
+		{
+			ControlRecord.bCachedHasStrengthOrDamping = false;
+			return;
+		}
+		else
+		{
+			ControlRecord.bCachedHasStrengthOrDamping = true;
+		}
+	}
+	else
+	{
+		if (!ControlRecord.bCachedHasStrengthOrDamping)
+		{
+			return;
+		}
+	}
+
+	// Update the target point on the child
+
+	// Danny TODO Note that if child is kinematic then getting the control point will be a problem
+	// as Chaos doesn't like to return the CoM for kinematics. Consider changing Chaos so that the
+	// CoM can still be queried. This workaround (of skipping the call) will result in bad behaviour
+	// if the child is kinematic and the parent is dynamic. 
+	if (!ChildActorHandle->GetIsKinematic())
+	{
+		// Note - This simply sets the location and doesn't trigger additional work
+		Constraint->SetChildConnectorLocation(
+			ControlComponent.ControlData.GetControlPoint(ChildActorHandle));
+	}
+
+	FTransform TargetTM(
+		ControlComponent.ControlTarget.TargetOrientation,
+		ControlComponent.ControlTarget.TargetPosition);
+
+	if (ControlComponent.ControlData.bUseSkeletalAnimation)
+	{
+		const FTransform ComponentSpaceAnimTargetTM = CalculateTargetTM(Hierarchy, JointSettings, ControlRecord);
+		const FTransform SimSpaceAnimTargetTM = ConvertComponentSpaceTransformToSimSpace(
+			SolverSettings, ComponentSpaceAnimTargetTM);
+		TargetTM = TargetTM * SimSpaceAnimTargetTM;
+	}
+	else
+	{
+		TargetTM = ConvertComponentSpaceTransformToSimSpace(SolverSettings, TargetTM);
+	}
+
+	Constraint->SetLinearDrivePositionTarget(TargetTM.GetTranslation());
+	Constraint->SetAngularDrivePositionTarget(TargetTM.GetRotation());
+
+	if ((ThisDeltaTime * ControlComponent.ControlData.LinearTargetVelocityMultiplier) != 0)
+	{
+		const FVector Velocity = UE::PhysicsControl::CalculateLinearVelocity(
+			ControlRecord.PreviousTargetTM.GetTranslation(), TargetTM.GetTranslation(), ThisDeltaTime);
+		Constraint->SetLinearDriveVelocityTarget(
+			Velocity * ControlComponent.ControlData.LinearTargetVelocityMultiplier);
+	}
+	else
+	{
+		Constraint->SetLinearDriveVelocityTarget(Chaos::FVec3(0));
+	}
+
+	if ((ThisDeltaTime * ControlComponent.ControlData.AngularTargetVelocityMultiplier) != 0)
+	{
+		const FQuat Q = TargetTM.GetRotation();
+		const FQuat PrevQ = ControlRecord.PreviousTargetTM.GetRotation();
+		const FVector AngularVelocity = UE::PhysicsControl::CalculateAngularVelocity(PrevQ, Q, ThisDeltaTime);
+		Constraint->SetAngularDriveVelocityTarget(
+			AngularVelocity * ControlComponent.ControlData.AngularTargetVelocityMultiplier);
+	}
+	else
+	{
+		Constraint->SetAngularDriveVelocityTarget(Chaos::FVec3(0));
+	}
+
+	// Danny TODO we don't actually have a shortcut for if the target is clean
+	ControlComponent.DirtyFlags &= ~ERigPhysicsControlComponentDirtyFlags::Target;
+
+	ControlRecord.PreviousTargetTM = TargetTM;
+	ControlRecord.PreviousTargetUpdateCounter = UpdateCounter;
+}
+
+//======================================================================================================================
+// Note that we read back into a target bone, which may have been specified explicitly, or will
+// otherwise default to the physics element parent.
+void FRigPhysicsSolver::UpdatePostPhysics(
+	URigHierarchy&              Hierarchy,
+	FRigPhysicsSolverComponent& SolverComponent, 
+	const float                 Alpha, 
+	const float                 DeltaTime,
+	const bool                  bUsingFixedStep)
+{
+	DECLARE_SCOPE_HIERARCHICAL_COUNTER_FUNC()
+	QUICK_SCOPE_CYCLE_COUNTER(STAT_RigPhysics_UpdatePostPhysics);
+
+	if (Alpha == 0.0f)
+	{
+		return;
+	}
+	const FRigPhysicsSolverSettings& SolverSettings = SolverComponent.SolverSettings;
+	bool bGotInvalidSimulationData = false;
+
+	double PositionThresholdForResetSquared = SolverSettings.PositionThresholdForReset > 0 ?
+		FMath::Square(SolverSettings.PositionThresholdForReset) : TNumericLimits<double>::Max();
+	double HighestPosition = -1.0;
+
+	// Traverse using the sorted keys
+	for (const FRigComponentKey& ComponentKey : SortedBodyComponentKeys)
+	{
+		FRigBodyRecord& Record = *BodyRecords.Find(ComponentKey);
+
+		if (Record.ActorHandle && Record.TargetElementKey.IsValid())
+		{
+			// Check the simulation output
+			FTransform SimSpaceTM = Record.ActorHandle->GetWorldTransform();
+			double DistSq = SimSpaceTM.GetTranslation().SquaredLength();
+			if (!SimSpaceTM.IsValid() || DistSq > PositionThresholdForResetSquared)
+			{
+				HighestPosition = FMath::Max(HighestPosition, FMath::Sqrt(DistSq));
+				bGotInvalidSimulationData = true;
+			}
+
+			// Calculate the target TM even if we're going to reset - it's likely useful for
+			// debugging (and this should be rare!)
+			Record.FinalComponentSpaceTM = ConvertSimSpaceTransformToComponentSpace(SolverSettings, SimSpaceTM);
+
+			// Danny TODO rework this - currently in does the interpolation here, and stores it in
+			// the component, but I think the component should store the actual physics transform,
+			// and interpolation should only happen when we write back into the bone transform.
+			if (Alpha < 0.999f)
+			{
+				// Danny TODO Note that this uses Alpha to blend in component space.
+				// This can cause joint separation. We probably want an option to blend
+				// in local (joint) space, perhaps splitting the alpha into orientation
+				// and position.
+				FQuat TargetQ = FQuat::Slerp(
+					Record.SourceComponentSpaceTM.GetRotation(), Record.FinalComponentSpaceTM.GetRotation(), Alpha);
+				FVector TargetT = FMath::Lerp(
+					Record.SourceComponentSpaceTM.GetTranslation(), Record.FinalComponentSpaceTM.GetTranslation(), Alpha);
+				Record.FinalComponentSpaceTM.SetRotation(TargetQ);
+				Record.FinalComponentSpaceTM.SetTranslation(TargetT);
+			}
+
+			if (FRigPhysicsBodyComponent* PhysicsBodyComponent = Cast<FRigPhysicsBodyComponent>(
+				Hierarchy.FindComponent(ComponentKey)))
+			{
+				PhysicsBodyComponent->ForceAndTorques.Empty();
+				SetPhysicsBodyComponentState(SolverSettings, Record, *PhysicsBodyComponent);
+			}
+		}
+	}
+
+	if (bGotInvalidSimulationData)
+	{
+		if (HighestPosition > 0)
+		{
+			UE_LOG(LogRigPhysics, Log, TEXT("Position %f triggered teleport - resetting pose"),
+				HighestPosition);
+		}
+		// Avoid cached transforms being used in controls by bumping the update counter. 
+		UpdateCounter += 1;
+		// Set this to 3 since it gets decremented at the end of the update, and we need it to take
+		// effect at the start of the next update.
+		SolverComponent.TrackInputCounter = FMath::Max(SolverComponent.TrackInputCounter, 3);
+
+		// If we found something invalid then we force the simulation to be as good as we can make it,
+		// and we don't write back to the hierarchy.
+		UE_LOG(LogRigPhysics, Log, TEXT("Resetting state to input pose"));
+		for (TPair<FRigComponentKey, FRigBodyRecord>& BodyRecordPair : BodyRecords)
+		{
+			const FRigComponentKey& ComponentKey = BodyRecordPair.Key;
+			const FRigBodyRecord& Record = BodyRecordPair.Value;
+
+			if (const FRigPhysicsBodyComponent* PhysicsBodyComponent = Cast<FRigPhysicsBodyComponent>(
+				Hierarchy.FindComponent(ComponentKey)))
+			{
+				if (Record.ActorHandle && !Record.ActorHandle->GetIsKinematic())
+				{
+					// Get the TM in component space, and then convert it into sim space.
+					FRigElementKey SourceKey = PhysicsBodyComponent->BodySolverSettings.SourceBone;
+					if (!SourceKey.IsValid())
+					{
+						SourceKey = ComponentKey.ElementKey;
+					}
+					if (SourceKey.IsValid())
+					{
+						const FTransform SourceSimulationSpaceTM =
+							ConvertComponentSpaceTransformToSimSpace(SolverSettings, Record.SourceComponentSpaceTM);
+						Record.ActorHandle->SetWorldTransform(SourceSimulationSpaceTM);
+					}
+					Record.ActorHandle->SetLinearVelocity(FVector::ZeroVector);
+					Record.ActorHandle->SetAngularVelocity(FVector::ZeroVector);
+				}
+			}
+		}
+	}
+	else
+	{
+		// All is good - write the transforms we cached
+		for (const FRigComponentKey& ComponentKey : SortedBodyComponentKeys)
+		{
+			const FRigBodyRecord& Record = *BodyRecords.Find(ComponentKey);
+			if (Record.ActorHandle)
+			{
+				// Kinematic velocities will be incorrect when using a fixed timestep, because of
+				// the rewind. Overwrite them here, since we know exactly what they should be, in
+				// case they are subsequently made dynamic. This is particularly important when
+				// warming up velocities following a reset.
+				if (bUsingFixedStep && Record.ActorHandle->GetIsKinematic())
+				{
+					FVector V = ConvertComponentSpaceVectorToSimSpace(
+						SolverSettings, Record.SourceComponentSpaceVelocity);
+					FVector W = ConvertComponentSpaceVectorToSimSpace(
+						SolverSettings, Record.SourceComponentSpaceAngularVelocity);
+					Record.ActorHandle->SetLinearVelocity(V);
+					Record.ActorHandle->SetAngularVelocity(W);
+				}
+
+				if (Record.TargetElementKey.IsValid())
+				{
+					// Note that we set bAffectChildren = true (i.e. don't counter-animate
+					// children), so that attached animation bones will follow physics, but we
+					// rely on our bodies being sorted so we work out from the leaf nodes so as
+					// not to disturb previously set bodies.
+					Hierarchy.SetGlobalTransform(Record.TargetElementKey, Record.FinalComponentSpaceTM, false, true);
+				}
+			}
+		}
+	}
+}
+
+//======================================================================================================================
+void FRigPhysicsSolver::StepSimulation(
+	const UWorld*                World,
+	const AActor*                OwningActorPtr,
+	const FRigVMExecuteContext&  ExecuteContext,
+	URigHierarchy&               Hierarchy,
+	FRigPhysicsSolverComponent&  SolverComponent,
+	const float                  DeltaTimeOverride,
+	const float                  SimulationSpaceDeltaTimeOverride,
+	const float                  Alpha)
+{
+	DECLARE_SCOPE_HIERARCHICAL_COUNTER_FUNC()
+	QUICK_SCOPE_CYCLE_COUNTER(STAT_RigPhysics_StepSimulation);
+
+	// Increment the update counter at the start - and always update it so this tells us our "frame number".
+	++UpdateCounter;
+
+	float PhysicsDeltaTime = ExecuteContext.GetDeltaTime();
+	if (DeltaTimeOverride > 0.0f)
+	{
+		PhysicsDeltaTime = DeltaTimeOverride;
+	}
+	else if (DeltaTimeOverride < 0.0f)
+	{
+		PhysicsDeltaTime = 0.0f;
+	}
+
+	float PhysicsSimulationSpaceDeltaTime = PhysicsDeltaTime;
+	if (SimulationSpaceDeltaTimeOverride > 0.0f)
+	{
+		PhysicsSimulationSpaceDeltaTime = SimulationSpaceDeltaTimeOverride;
+	}
+
+	// Lock against access to the simulation, in case of access by the WorldObject task
+	UE::TScopeLock<FTransactionallySafeCriticalSection> Lock(*SimulationMutex.Get());
+
+	// We need to know about the simulation space etc before we can instantiate anything into the right place
+	SimulationSpaceData = UpdateSimulationSpaceStateAndCalculateData(
+		ExecuteContext, Hierarchy, SolverComponent, PhysicsSimulationSpaceDeltaTime, ExecuteContext.GetAbsoluteTime());
+
+	// We instantiate when we do the first simulation - this makes sure any changes applied by
+	// the user have been made. It also means there is no overhead if physics is never stepped.
+	// However, there may be a hitch due to the creation, so it may also happen during construction.
+	Instantiate(ExecuteContext, Hierarchy, SolverComponent); // There is an early out if it's already been done
+
+	if (!Simulation)
+	{
+		return;
+	}
+
+	const FRigPhysicsSolverSettings& SolverSettings = SolverComponent.SolverSettings;
+	const FRigPhysicsSimulationSpaceSettings& SimulationSpaceSettings = SolverComponent.SimulationSpaceSettings;
+
+	float FixedTimeStep = CVarControlRigPhysicsFixedTimeStepOverride.GetValueOnAnyThread() < 0 
+		? SolverSettings.FixedTimeStep : CVarControlRigPhysicsFixedTimeStepOverride.GetValueOnAnyThread();
+	int MaxTimeSteps = CVarControlRigPhysicsMaxTimeStepsOverride.GetValueOnAnyThread() < 0
+		? SolverSettings.MaxTimeSteps : CVarControlRigPhysicsMaxTimeStepsOverride.GetValueOnAnyThread();
+	float MaxDeltaTime = CVarControlRigPhysicsMaxDeltaTimeOverride.GetValueOnAnyThread() < 0
+		? SolverSettings.MaxDeltaTime : CVarControlRigPhysicsMaxDeltaTimeOverride.GetValueOnAnyThread();
+
+	// Set settings that might change
+	Simulation->SetSolverSettings(
+		FixedTimeStep,
+		SolverSettings.CollisionBoundsExpansion,
+		SolverSettings.MaxDepenetrationVelocity,
+		SolverSettings.bUseLinearJointSolver,
+		SolverSettings.PositionIterations,
+		SolverSettings.VelocityIterations,
+		SolverSettings.ProjectionIterations,
+		SolverSettings.bUseManifolds);
+
+	Simulation->SetUseMinStepTime(false);
+	Simulation->SetUseFixedStepTolerance(false);
+
+	Chaos::FCollisionDetectorSettings CollisionDetectorSettings = Simulation->GetCollisionDetectorSettings();
+	CollisionDetectorSettings.BoundsVelocityInflation = SolverSettings.BoundsVelocityMultiplier;
+	CollisionDetectorSettings.MaxVelocityBoundsExpansion = SolverSettings.MaxVelocityBoundsExpansion;
+	CollisionDetectorSettings.bAllowCCD = SolverSettings.bAllowCCD;
+	Simulation->SetCollisionDetectorSettings(CollisionDetectorSettings);
+
+	// This gets reset to 100 after every simulation step!
+	Simulation->SetMaxNumRollingAverageStepTimes(SolverSettings.MaxNumRollingAverageStepTimes);
+
+	// Other settings - would normally be static (so Danny TODO move this)
+	ChaosJointSolverSettings.bSolvePositionLast = SolverSettings.bSolveJointPositionsLast;
+	ChaosJointSolverSettings.bSortEnabled = true;
+
+	// Simulation space
+	Chaos::FSimulationSpaceSettings ChaosSimulationSpaceSettings = Simulation->GetSimulationSpaceSettings();
+	ChaosSimulationSpaceSettings.bEnabled = (SimulationSpaceSettings.SpaceMovementAmount > 0.0f);
+	ChaosSimulationSpaceSettings.ExternalLinearEtherDrag = SimulationSpaceSettings.ExternalLinearDrag;
+	ChaosSimulationSpaceSettings.LinearVelocityAlpha = SimulationSpaceSettings.LinearDragMultiplier;
+	ChaosSimulationSpaceSettings.AngularVelocityAlpha = SimulationSpaceSettings.AngularDragMultiplier;
+	Simulation->SetSimulationSpaceSettings(ChaosSimulationSpaceSettings);
+	Simulation->UpdateSimulationSpace(
+		SimulationSpaceState.SimulationSpaceTM,
+		SimulationSpaceSettings.SpaceMovementAmount * SimulationSpaceData.LinearVelocity,
+		SimulationSpaceSettings.SpaceMovementAmount * SimulationSpaceData.AngularVelocity,
+		SimulationSpaceSettings.SpaceMovementAmount * SimulationSpaceData.LinearAcceleration,
+		SimulationSpaceSettings.SpaceMovementAmount * SimulationSpaceData.AngularAcceleration);
+
+	UpdateWorldObjectsPrePhysics(SolverSettings);
+
+	// Only update if there is a delta time:
+	// * We don't want to update our previous TMs and store the dt - because that would end up
+	//   implying infinite velocities
+	// * We don't to update kinematic bodies with the new TMs because, since the simulated ones
+	//   won't move, that would break the pose.
+	// * We can't actually simulate with dt = 0
+	if (PhysicsDeltaTime > 0.0f)
+	{
+		UpdatePrePhysics(ExecuteContext, Hierarchy, SolverComponent, PhysicsDeltaTime);
+
+		Simulation->Simulate(
+			PhysicsDeltaTime, MaxDeltaTime, MaxTimeSteps,
+			SimulationSpaceData.Gravity, &ChaosJointSolverSettings);
+
+		PreviousUpdateCounter = UpdateCounter;
+	}
+
+	// Always do a read-back, even for zero Dt
+	UpdatePostPhysics(Hierarchy, SolverComponent, FMath::Clamp(Alpha, 0.0f, 1.0f), PhysicsDeltaTime, FixedTimeStep > 0);
+
+	// Trigger the update of world objects as soon as possible. Note that this needs to be done
+	// after the last access to the simulation, because the game-thread task will start
+	// adding/removing objects from our simulation, and we can't let that work overlap with what we
+	// do here.
+
+	UpdateWorldObjectsPostPhysics(World, SolverSettings, OwningActorPtr);
+
+	if (SolverComponent.TrackInputCounter > 0)
+	{
+		--SolverComponent.TrackInputCounter;
+	}
+}
diff --git a/Engine/Plugins/Experimental/ControlRigPhysics/Source/ControlRigPhysics/Private/RigPhysicsSolver.h b/Engine/Plugins/Experimental/ControlRigPhysics/Source/ControlRigPhysics/Private/RigPhysicsSolver.h
new file mode 100644
--- /dev/null
+++ b/Engine/Plugins/Experimental/ControlRigPhysics/Source/ControlRigPhysics/Private/RigPhysicsSolver.h
@@ -0,0 +1,496 @@
+// Copyright Epic Games, Inc. All Rights Reserved.
+
+#pragma once
+
+#include "RigPhysicsData.h"
+
+#include "PhysicsControlPoseData.h"
+
+#include "Rigs/RigHierarchyComponents.h"
+
+#include "Physics/ImmediatePhysics/ImmediatePhysicsDeclares.h"
+#include "Physics/ImmediatePhysics/ImmediatePhysicsSimulation.h"
+#include "Chaos/PBDJointConstraintTypes.h"
+#include "Misc/TransactionallySafeCriticalSection.h"
+
+#include "RigPhysicsSolver.generated.h"
+
+class UControlRig;
+
+struct FRigVMDrawInterface;
+struct FRigVMExecuteContext;
+struct FRigPhysicsSolverComponent;
+
+struct FRigPhysicsCollision;
+struct FRigPhysicsDynamics;
+struct FRigPhysicsBodyComponent;
+struct FRigPhysicsControlComponent;
+struct FRigPhysicsJointComponent;
+
+#if UE_BUILD_SHIPPING || UE_BUILD_TEST
+DECLARE_LOG_CATEGORY_EXTERN(LogRigPhysics, Warning, Warning);
+#else
+DECLARE_LOG_CATEGORY_EXTERN(LogRigPhysics, Log, All);
+#endif
+
+//======================================================================================================================
+// For internal use only
+//======================================================================================================================
+struct FRigPhysicsIgnorePair
+{
+	FRigPhysicsIgnorePair(const FRigComponentKey& InA, const FRigComponentKey& InB)
+		: A(InA), B(InB) {}
+	FRigComponentKey A;
+	FRigComponentKey B;
+
+	bool operator==(const FRigPhysicsIgnorePair& Other) const
+	{
+		return (A == Other.A && B == Other.B) || (A == Other.B && B == Other.A);
+	}
+};
+
+FORCEINLINE uint32 GetTypeHash(const FRigPhysicsIgnorePair& Pair)
+{
+	return HashCombine(GetTypeHash(FMath::Min(Pair.A, Pair.B)), GetTypeHash(FMath::Max(Pair.A, Pair.B)));
+}
+typedef TSet<FRigPhysicsIgnorePair> FRigPhysicsIgnorePairs;
+
+//======================================================================================================================
+// For internal use only, to keep track of objects we have created
+//======================================================================================================================
+USTRUCT()
+struct FRigBodyRecord
+{
+	GENERATED_BODY()
+
+	// Things that are set during instantiation
+	ImmediatePhysics::FActorHandle* ActorHandle = nullptr;
+
+	// Cache the element key for where we will write the simulation result
+	FRigElementKey TargetElementKey;
+
+	// The final/simulated transform is stored before writing it into the output, so we can avoid
+	// corrupting the output if anything is bad and we need to reset.
+	FTransform FinalComponentSpaceTM;
+
+	// These source (i.e. bone/element) transforms are updated for all records in UpdatePrePhysics.
+	// The times/validity are determined by the CurrentDeltaTime, PrevDeltaTime and update counters
+	// in the simulation itself.
+	FTransform SourceComponentSpaceTM;
+	FVector    SourceComponentSpaceVelocity = FVector::ZeroVector;
+	FVector    SourceComponentSpaceAngularVelocity = FVector::ZeroVector;
+
+	FTransform PrevSourceComponentSpaceTM;
+	FVector    PrevSourceComponentSpaceVelocity = FVector::ZeroVector;
+	FVector    PrevSourceComponentSpaceAngularVelocity = FVector::ZeroVector;
+};
+
+//======================================================================================================================
+// For internal use only, to keep track of joints we have created
+//======================================================================================================================
+USTRUCT()
+struct FRigJointRecord
+{
+	GENERATED_BODY()
+
+	// Things that are set during instantiation
+	ImmediatePhysics::FJointHandle* JointHandle = nullptr;
+
+	// These keys are filled in when the record is created, even if the original key is set to pick
+	// up the components automatically.
+	FRigComponentKey                ParentBodyComponentKey;
+	FRigComponentKey                ChildBodyComponentKey;
+
+	// Things that are updated as the simulation progresses
+
+	// The drive works with velocities so we store the previous target transform, and when it was stored.
+	FTransform                      PreviousDriveTargetTM;
+	// This is stored from the main solver update counter, marking when the previous drive TM was valid.
+	int64                           PreviousDriveTargetUpdateCounter = -999;
+};
+
+//======================================================================================================================
+// For internal use only, to keep track of controls we have created
+//======================================================================================================================
+USTRUCT()
+struct FRigControlRecord
+{
+	GENERATED_BODY()
+
+	// Things that are set during instantiation
+	ImmediatePhysics::FJointHandle* JointHandle = nullptr;
+
+	// These keys are filled in when the record is created, even if the original key is set to pick
+	// up the components automatically.
+	FRigComponentKey                ParentBodyComponentKey;
+	FRigComponentKey                ChildBodyComponentKey;
+
+	// Things that are updated as the simulation progresses
+
+	// The control works with velocities so we store the previous target, and when it was stored.
+	FTransform                      PreviousTargetTM;
+	// This is stored from the main solver update counter, marking when the previous target TM was valid.
+	int64                           PreviousTargetUpdateCounter = -999;
+
+	// When the FRigPhysicsControlComponent is dirty, this is calculated from the strengths etc, and
+	// then cached. Then if FRigPhysicsControlComponent is clean, we can see whether there is any
+	// strength/damping that would require the target to be updated.
+	bool bCachedHasStrengthOrDamping = true;
+};
+
+//======================================================================================================================
+// This represents the low level simulation, plus all the objects and controls we make to go in it
+//======================================================================================================================
+USTRUCT()
+struct FRigPhysicsSolver
+{
+public:
+	GENERATED_BODY()
+
+	FRigPhysicsSolver() {}
+	FRigPhysicsSolver(const FRigComponentKey& SolverComponentKey);
+
+	// This will initialise/create the simulation and then create everything we need in it. 
+	void Instantiate(
+		const FRigVMExecuteContext&       ExecuteContext, 
+		URigHierarchy&                    Hierarchy, 
+		const FRigPhysicsSolverComponent& SolverComponent);
+
+	// Integrates the simulation forwards.
+	// If DeltaTimeOverride is +ve, then that value is used.
+	// If it is zero, then delta time is taken from the execute context
+	// If it is negative, then the simulation isn't stepped.
+	void StepSimulation(
+		const UWorld*                World,
+		const AActor*                OwningActorPtr,
+		const FRigVMExecuteContext&  ExecuteContext, 
+		URigHierarchy&               Hierarchy,
+		FRigPhysicsSolverComponent&  SolverComponent,
+		const float                  DeltaTimeOverride, 
+		const float                  SimulationSpaceDeltaTimeOverride,
+		const float                  Alpha);
+
+	// Draws shapes etc AND (potentially) enables the low-level Chaos debug draw
+	void Draw(
+		FRigVMDrawInterface*                    DI,
+		const FRigPhysicsSolverSettings&        SolverSettings,
+		const FRigPhysicsVisualizationSettings& VisualizationSettings,
+		const UWorld*                           DebugWorld) const;
+
+	// Represents the properties of the simulation space - calculated near the beginning of the update.
+	// Note that all these are specified in the simulation space itself
+	struct FSimulationSpaceData
+	{
+		FVector LinearVelocity;
+		// Angular velocity is in rad/s
+		FVector AngularVelocity;
+		FVector LinearAcceleration;
+		// Angular acceleration is in rad/s/s
+		FVector AngularAcceleration;
+		FVector Gravity;
+	};
+
+	// Returns the simulation space data, as calculated at the start of the last step
+	const FSimulationSpaceData& GetSimulationSpaceData() const { return SimulationSpaceData; }
+
+	// Returns the component that we're associated with/owned by
+	const FRigComponentKey& GetSolverComponentKey() const { return PhysicsSolverComponentKey; }
+
+private:
+	// Used by the world-space to simulation-space motion transfer system in Component- or
+	// Bone-Space sims, and preserved between updates.
+	struct FSimulationSpaceState
+	{
+		FTransform ComponentTM;
+		FTransform BoneRelComponentTM;
+
+		// The world transform of the simulation space
+		FTransform SimulationSpaceTM;
+		FTransform PrevSimulationSpaceTM;
+		FTransform PrevPrevSimulationSpaceTM;
+		// The time between SimulationSpaceTM and PrevSimulationSpaceTM
+		float      Dt = 1.0f; 
+		// The time between PrevSimulationSpaceTM and PrevPrevSimulationSpaceTM
+		float      PrevDt = 1.0f;
+	};
+
+	// Creates the low level simulation
+	void InitialiseSimulation();
+
+	// This collects all the bodies associated with the solver and makes records for them
+	void InitialiseBodyRecords(const URigHierarchy& Hierarchy, const FRigPhysicsSolverComponent& SolverComponent);
+
+	// This collects all the joints associated with the solver and makes records for them
+	void InitialiseJointRecords(const URigHierarchy& Hierarchy, const FRigPhysicsSolverComponent& SolverComponent);
+
+	// This collects all the controls associated with the solver and makes records for them
+	void InitialiseControlRecords(const URigHierarchy& Hierarchy, const FRigPhysicsSolverComponent& SolverComponent);
+
+	// Creates low level bodies
+	void InstantiatePhysicsBodies(
+		URigHierarchy&                    Hierarchy,
+		const FRigPhysicsSolverComponent& SolverComponent,
+		FRigPhysicsIgnorePairs&           IgnorePairs);
+
+	// Creates low level physics joints 
+	void InstantiatePhysicsJoints(
+		const URigHierarchy&              Hierarchy,
+		const FRigPhysicsSolverComponent& SolverComponent,
+		FRigPhysicsIgnorePairs&           IgnorePairs);
+
+	// Creates low level controls
+	void InstantiateControls(
+		const URigHierarchy&              Hierarchy,
+		const FRigPhysicsSolverComponent& SolverComponent,
+		FRigPhysicsIgnorePairs&           IgnorePairs);
+
+	// Creates collision shapes associated with the solver. Also applies IgnorePairs
+	void InstantiateSolverCollision(
+		const FRigPhysicsSolverComponent& SolverComponent,
+		FRigPhysicsIgnorePairs&           IgnorePairs);
+
+	// Returns true if the component is physics, and its solver matches the/our solver component
+	// (directly, or automatically)
+	bool ShouldComponentBeInSimulation(
+		const URigHierarchy&              Hierarchy, 
+		const FRigPhysicsSolverComponent& SolverComponent, 
+		const FRigComponentKey&           ComponentKey) const;
+
+	// Creates an actor with collision. This will be dynamic if Dynamics is valid, or otherwise kinematic
+	ImmediatePhysics::FActorHandle* CreateBody(
+		const FName                        BodyName,
+		const FRigPhysicsCollision&        Collision,
+		const FRigPhysicsDynamics*         Dynamics,
+		const FPhysicsControlModifierData* BodyData,
+		const FTransform&                  BodyRelSimSpaceTM) const;
+
+	// This releases the actors etc in the simulation, destroys the simulation, and resets our
+	// internal records.
+	void DestroyPhysicsSimulation();
+
+	void UpdatePrePhysics(
+		const FRigVMExecuteContext&        ExecuteContext, 
+		const URigHierarchy&               Hierarchy,
+		FRigPhysicsSolverComponent&        SolverComponent, 
+		const float                        DeltaTime);
+
+	void UpdateBodyRecordPrePhysics(
+		const URigHierarchy&               Hierarchy,
+		const FRigPhysicsSolverComponent&  SolverComponent,
+		const float                        DeltaTime,
+		FRigBodyRecord&                    Record,
+		const FRigPhysicsBodyComponent*    PhysicsComponent);
+
+	void UpdateBodyPrePhysics(
+		const FRigVMExecuteContext&        ExecuteContext, 
+		const FRigPhysicsSolverComponent&  SolverComponent,
+		const FRigBodyRecord&              Record, 
+		const FRigPhysicsBodyComponent&    PhysicsComponent);
+
+	void ApplyForcesPrePhysics(
+		const FRigVMExecuteContext&        ExecuteContext, 
+		const FRigPhysicsSolverComponent&  SolverComponent,
+		const FRigBodyRecord&              Record, 
+		const FRigPhysicsBodyComponent&    PhysicsComponent);
+
+	void UpdateJointPrePhysics(
+		const URigHierarchy&               Hierarchy, 
+		FRigJointRecord&                   JointRecord,
+		const FRigPhysicsJointComponent&   JointComponent, 
+		const float                        DeltaTime);
+
+	void UpdateControlPrePhysics(
+		FRigControlRecord&                 ControlRecord,
+		const FRigPhysicsControlComponent& ControlComponent,
+		const FRigPhysicsSolverComponent&  SolverComponent,
+		const URigHierarchy&               Hierarchy,
+		const float                        DeltaTime);
+
+	void UpdatePostPhysics(
+		URigHierarchy&              Hierarchy,
+		FRigPhysicsSolverComponent& SolverComponent, 
+		const float                 Alpha, 
+		const float                 DeltaTime,
+		const bool                  bUsingFixedStep);
+
+	// This will walk through the WorldObjects, which should now be updated, and update the simulation
+	void UpdateWorldObjectsPrePhysics(const FRigPhysicsSolverSettings& SolverSettings);
+
+	// This will trigger a game-thread update of WorldObjects, ready for the next update
+	void UpdateWorldObjectsPostPhysics(
+		const UWorld*                    World,
+		const FRigPhysicsSolverSettings& SolverSetting, 
+		const AActor*                    OwningActorPtr);
+
+	void InitSimulationSpace(const FTransform& ComponentToWorld, const FTransform& BoneToComponent);
+
+	void CheckForResetsPrePhysics(
+		const URigHierarchy& Hierarchy, FRigPhysicsSolverComponent& SolverComponent, const float DeltaTime);
+
+	// AbsoluteTime is used for calculating turbulence when applying the "wind"
+	FSimulationSpaceData UpdateSimulationSpaceStateAndCalculateData(
+		const FRigVMExecuteContext&         ExecuteContext,
+		const URigHierarchy&                Hierarchy,
+		const FRigPhysicsSolverComponent&   SolverComponent,
+		const float                         Dt,
+		const double                        AbsoluteTime);
+
+	static FTransform GetSpaceTransform(
+		ERigPhysicsSimulationSpace Space, const FTransform& ComponentTM, const FTransform& BoneTM);
+
+	// Returns the simulation space transform, in world space
+	FTransform GetSimulationSpaceTransform(const FRigPhysicsSolverSettings& SolverSettings) const;
+
+	// Converts a transform from component space (e.g. coming from the owning control rig) into the
+	// simulation space
+	FTransform ConvertComponentSpaceTransformToSimSpace(
+		const FRigPhysicsSolverSettings& SolverSettings, const FTransform& ComponentSpaceTM) const;
+
+	// Converts a component space position to simulation space, including translation
+	FVector ConvertComponentSpacePositionToSimSpace(
+		const FRigPhysicsSolverSettings& SolverSettings, const FVector& ComponentSpacePosition) const;
+
+	// Converts a component space vector to simulation space, ignoring translation
+	FVector ConvertComponentSpaceVectorToSimSpace(
+		const FRigPhysicsSolverSettings& SolverSettings, const FVector& ComponentSpaceVector) const;
+
+	// Converts a transform specified in world space to the simulation space
+	FTransform ConvertWorldTransformToSimSpace(
+		const FRigPhysicsSolverSettings& SolverSettings, const FTransform& WorldSpaceTM) const;
+
+	// Converts a position specified in world space to the simulation space
+	FVector ConvertWorldPositionToSimSpaceNoScale(
+		const FRigPhysicsSolverSettings& SolverSettings, const FVector& WorldSpacePosition) const;
+
+	// Converts a vector specified in world space to the simulation space (e.g. converting gravity)
+	FVector ConvertWorldVectorToSimSpaceNoScale(
+		const FRigPhysicsSolverSettings& SolverSettings, const FVector& WorldSpaceVector) const;
+
+	// Converts a transform from the simulation space to component space (e.g. for writing back to
+	// the owning control rig)
+	FTransform ConvertSimSpaceTransformToComponentSpace(
+		const FRigPhysicsSolverSettings& SolverSettings, const FTransform& SimSpaceTM) const;
+
+	// Converts a simulation space position to component space, including translation
+	FVector ConvertSimSpacePositionToComponentSpace(
+		const FRigPhysicsSolverSettings& SolverSettings, const FVector& SimSpacePosition) const;
+
+	// Converts a simulation space vector to component space (ignoring translation)
+	FVector ConvertSimSpaceVectorToComponentSpace(
+		const FRigPhysicsSolverSettings& SolverSettings, const FVector& SimSpaceVector) const;
+
+	// Converts a collision space transform to simulation space
+	FTransform ConvertCollisionSpaceTransformToSimSpace(
+		const FRigPhysicsSolverSettings& SolverSettings, const FTransform& TM) const;
+
+	// Gets the simulation actor handle for a component key. Note that the component key could be a
+	// body or a solver component. Can return nullptr
+	ImmediatePhysics::FActorHandle* GetActor(const FRigComponentKey& ComponentKey) const;
+
+	// Updates the PhysicsBodyComponent with the state of the rigid body from the record. Assumes
+	// that FinalComponentSpaceTM has been set.
+	void SetPhysicsBodyComponentState(
+		const FRigPhysicsSolverSettings& SolverSettings, 
+		FRigBodyRecord&                  Record, 
+		FRigPhysicsBodyComponent&        PhysicsBodyComponent) const;
+
+private:
+	FRigComponentKey        PhysicsSolverComponentKey;
+
+	// All the bodies, but in no particular order
+	TMap<FRigComponentKey, FRigBodyRecord> BodyRecords;
+
+	// Ordering so that we can traverse from root to leaf bones
+	TArray<FRigComponentKey> SortedBodyComponentKeys;
+
+	// All the joints
+	TMap<FRigComponentKey, FRigJointRecord> JointRecords;
+
+	// All the controls
+	TMap<FRigComponentKey, FRigControlRecord> ControlRecords;
+
+	TSharedPtr<ImmediatePhysics::FSimulation> Simulation;
+
+	// Used to store things the simulation collision shape. May be offset from the origin if
+	// collision is in a different space to the simulation.
+	ImmediatePhysics::FActorHandle* CollisionActorHandle = nullptr;
+
+	// Used to make controls when they're not attached to another simulated body. Will always be at
+	// the origin.
+	ImmediatePhysics::FActorHandle* SimulationActorHandle = nullptr;
+
+	Chaos::FPBDJointSolverSettings ChaosJointSolverSettings;
+
+	FSimulationSpaceState SimulationSpaceState;
+
+	// Retain the data - we don't actually need to but (a) it makes it available for debugging and
+	// (b) it avoids passing it through the functions.
+	FSimulationSpaceData SimulationSpaceData;
+
+	// This is incremented at the very beginning of each simulation step, used to identify when previously
+	// calculated values are valid. Note that the universe will roll over before a int64 does
+	int64 UpdateCounter = 0;
+
+	// The update counter when PrevSourceComponentSpaceTM etc were written. Check this before using them
+	int64 PreviousUpdateCounter = -999;
+
+	// This will be set to false after instantiation, which can happen manually during construction,
+	// or be postponed until there is a simulation step.
+	bool bNeedToInstantiate = true;
+
+	// Record of a world object that we want to track, and "import" into our simulation for collision.
+	struct FWorldObject
+	{
+		FWorldObject() : ActorHandle(nullptr), LastSeenUpdateCounter(0) {}
+		FWorldObject(ImmediatePhysics::FActorHandle* InActorHandle, int64 InLastSeenUpdateCounter)
+			: ActorHandle(InActorHandle)
+			, LastSeenUpdateCounter(InLastSeenUpdateCounter) {
+		}
+
+		// The object we are tracking
+		TWeakObjectPtr<UPrimitiveComponent> WorldPrimitiveComponent;
+
+		// The actor in our simulation
+		ImmediatePhysics::FActorHandle* ActorHandle;
+
+		// The transform of the world object. This will be updated in the game-thread task.
+		FTransform ComponentWorldTransform;
+
+		// When the world object was last seen (i.e. captured by an overlap), used to detect when we
+		// should expire/forget about it.
+		int64 LastSeenUpdateCounter;
+
+		// For instanced static mesh components, the index of the specific instance this actor
+		// represents. INDEX_NONE for non-instanced components.
+		int32 InstanceBodyIndex = INDEX_NONE;
+
+		// This returns true if the object hasn't been seen recently, or is invalid
+		bool GetExpired(int64 InUpdateCounter, int64 InRetainCount) const
+		{
+			if (LastSeenUpdateCounter != -1 && InUpdateCounter > LastSeenUpdateCounter + InRetainCount)
+			{
+				return true;
+			}
+			if (!WorldPrimitiveComponent.IsValid() || !ActorHandle)
+			{
+				return true;
+			}
+			return false;
+		}
+	};
+
+	// The world objects we track for collision. Index uses GetUniqueID from the PrimitiveComponent.
+	// Note that this is only guaranteed to be unique whilst the object exists, so it is possible
+	// (but hopefully unlikely) that an object will be replaced and have the same ID. This will be
+	// rectified on the next update, as the record's pointer to the primitive component will be invalid.
+	TSharedPtr<TMap<uint32, FWorldObject>> WorldObjects;
+
+	// The box used for overlap tests to find the world objects
+	FBox WorldOverlapBox;
+
+	// Guard against access to the low level simulation. Note that this access may come from an
+	// async/spawned task, so the mutex itself needs to be shared, as we (as primary owners) may get
+	// deleted whilst the spawned task is in flight.
+	TSharedPtr<FTransactionallySafeCriticalSection> SimulationMutex;
+};
+
diff --git a/Engine/Plugins/Experimental/ControlRigPhysics/Source/ControlRigPhysics/Private/RigPhysicsSolverComponent.cpp b/Engine/Plugins/Experimental/ControlRigPhysics/Source/ControlRigPhysics/Private/RigPhysicsSolverComponent.cpp
--- a/Engine/Plugins/Experimental/ControlRigPhysics/Source/ControlRigPhysics/Private/RigPhysicsSolverComponent.cpp
+++ b/Engine/Plugins/Experimental/ControlRigPhysics/Source/ControlRigPhysics/Private/RigPhysicsSolverComponent.cpp
@@ -2,14 +2,16 @@
 
 #include "RigPhysicsSolverComponent.h"
 #include "PhysicsControlObjectVersion.h"
+#include "RigPhysicsSolver.h"
+
 #include "Rigs/RigHierarchy.h"
 #if WITH_EDITOR
 #include "ControlRigPhysicsEditorStyle.h"
 #endif
 
-//======================================================================================================================
-
 #include UE_INLINE_GENERATED_CPP_BY_NAME(RigPhysicsSolverComponent)
+
+//======================================================================================================================
 void FRigPhysicsSolverComponent::Save(FArchive& Ar)
 {
 	Ar.UsingCustomVersion(FPhysicsControlObjectVersion::GUID);
@@ -41,6 +43,19 @@ void FRigPhysicsSolverComponent::OnAddedToHierarchy(URigHierarchy* InHierarchy,
 	}
 }
 
+//======================================================================================================================
+FRigPhysicsSolver* FRigPhysicsSolverComponent::GetPhysicsSolver() const
+{
+	if (!PhysicsSolverPtr.IsValid())
+	{
+		const FRigComponentKey SolverComponentKey = GetKey();
+		PhysicsSolverPtr = MakeShared<FRigPhysicsSolver>(SolverComponentKey);
+	}
+	FRigPhysicsSolver* Result = PhysicsSolverPtr.Get();
+	check(Result->GetSolverComponentKey() == GetKey());
+	return Result;
+}
+
 //======================================================================================================================
 #if WITH_EDITOR
 const FSlateIcon& FRigPhysicsSolverComponent::GetIconForUI() const
diff --git a/Engine/Plugins/Experimental/ControlRigPhysics/Source/ControlRigPhysics/Private/RigPhysicsSolverExecution.cpp b/Engine/Plugins/Experimental/ControlRigPhysics/Source/ControlRigPhysics/Private/RigPhysicsSolverExecution.cpp
--- a/Engine/Plugins/Experimental/ControlRigPhysics/Source/ControlRigPhysics/Private/RigPhysicsSolverExecution.cpp
+++ b/Engine/Plugins/Experimental/ControlRigPhysics/Source/ControlRigPhysics/Private/RigPhysicsSolverExecution.cpp
@@ -2,7 +2,7 @@
 
 #include "RigPhysicsSolverExecution.h"
 
-#include "RigPhysicsSimulation.h"
+#include "RigPhysicsSolver.h"
 #include "ControlRigPhysicsModule.h"
 
 #include "Rigs/RigHierarchyController.h"
@@ -53,6 +53,28 @@ FRigUnit_AddPhysicsSolver_Execute()
 	}
 }
 
+//======================================================================================================================
+FRigUnit_InstantiatePhysics_Execute()
+{
+	DECLARE_SCOPE_HIERARCHICAL_COUNTER_RIGUNIT();
+	if (!ExecuteContext.Hierarchy)
+	{
+		return;
+	}
+	URigHierarchy& Hierarchy = *ExecuteContext.Hierarchy;
+	if (UControlRig* ControlRig = Cast<UControlRig>(Hierarchy.GetOuter()))
+	{
+		if (FRigPhysicsSolverComponent* SolverComponent = Cast<FRigPhysicsSolverComponent>(
+			ExecuteContext.Hierarchy->FindComponent(PhysicsSolverComponentKey)))
+		{
+			if (FRigPhysicsSolver* PhysicsSolver = SolverComponent->GetPhysicsSolver())
+			{
+				PhysicsSolver->Instantiate(ExecuteContext, Hierarchy, *SolverComponent);
+			}
+		}
+	}
+}
+
 //======================================================================================================================
 FRigUnit_TrackInputPose_Execute()
 {
@@ -79,6 +101,8 @@ FRigUnit_TrackInputPose_Execute()
 FRigUnit_StepPhysicsSolver_Execute()
 {
 	DECLARE_SCOPE_HIERARCHICAL_COUNTER_RIGUNIT();
+	SCOPED_NAMED_EVENT(RigPhysicsSolver_StepSimulation, FColor::Yellow);
+
 	if (!ExecuteContext.Hierarchy)
 	{
 		return;
@@ -87,41 +111,33 @@ FRigUnit_StepPhysicsSolver_Execute()
 	{
 		return;
 	}
+
 	URigHierarchy& Hierarchy = *ExecuteContext.Hierarchy;
-	if (UControlRig* ControlRig = Cast<UControlRig>(Hierarchy.GetOuter()))
-	{
-		if (FRigPhysicsSolverComponent* SolverComponent = Cast<FRigPhysicsSolverComponent>(
-			Hierarchy.FindComponent(PhysicsSolverComponentKey)))
-		{
-			FRigPhysicsSimulation* PhysicsSimulation = Cast<FRigPhysicsSimulation>(
-				ControlRig->GetPhysicsSimulation(PhysicsSolverComponentKey));
 
-			// Make the solver and instantiate if necessary
-			if (!PhysicsSimulation)
-			{
-				TSharedPtr<FRigPhysicsSimulation> NewPhysicsSimulation = MakeShared<FRigPhysicsSimulation>(
-					FName(ControlRig->GetName()));
+	FRigPhysicsSolverComponent* SolverComponent = Cast<FRigPhysicsSolverComponent>(
+		Hierarchy.FindComponent(PhysicsSolverComponentKey));
 
-				if (ControlRig->RegisterPhysicsSimulation(NewPhysicsSimulation, PhysicsSolverComponentKey))
-				{
-					NewPhysicsSimulation->Instantiate(ExecuteContext, Hierarchy, *SolverComponent);
-					PhysicsSimulation = NewPhysicsSimulation.Get();
-				}
-			}
+	if (SolverComponent)
+	{
+		if (UControlRig* ControlRig = Cast<UControlRig>(Hierarchy.GetOuter()))
+		{
+			// We only set this flag once we start executing, because that's when it matters.
+			// Control rig is responsible for clearing it when the rig is constructed.
+			ControlRig->SetContainsSimulation();
 
-			if (PhysicsSimulation)
+			if (FRigPhysicsSolver* PhysicsSolver = SolverComponent->GetPhysicsSolver())
 			{
 				const USceneComponent* OwningSceneComponent = ControlRig->GetOwningSceneComponent();
 				const AActor* OwningActorPtr = OwningSceneComponent ? OwningSceneComponent->GetOwner() : nullptr;
 
-				PhysicsSimulation->StepSimulation(
+				PhysicsSolver->StepSimulation(
 					ControlRig->GetWorld(), OwningActorPtr, ExecuteContext, Hierarchy, *SolverComponent, 
 					DeltaTimeOverride, SimulationSpaceDeltaTimeOverride, Alpha);
 				int32 VisualizationOverride = CVarControlRigPhysicsVisualizationOverride.GetValueOnAnyThread();
 				if (VisualizationOverride == 1 || 
 					(VisualizationOverride < 0 && VisualizationSettings.bEnableVisualization))
 				{
-					PhysicsSimulation->Draw(
+					PhysicsSolver->Draw(
 						ExecuteContext.GetDrawInterface(), SolverComponent->SolverSettings,
 						VisualizationSettings, ExecuteContext.GetWorld());
 				}
@@ -188,23 +204,19 @@ FRigUnit_GetPhysicsSolverSpaceData_Execute()
 	{
 		return;
 	}
-	if (UControlRig* ControlRig = Cast<UControlRig>(ExecuteContext.Hierarchy->GetOuter()))
+	if (FRigPhysicsSolverComponent* SolverComponent = Cast<FRigPhysicsSolverComponent>(
+		ExecuteContext.Hierarchy->FindComponent(PhysicsSolverComponentKey)))
 	{
-		if (FRigPhysicsSolverComponent* SolverComponent = Cast<FRigPhysicsSolverComponent>(
-			ExecuteContext.Hierarchy->FindComponent(PhysicsSolverComponentKey)))
+		if (FRigPhysicsSolver* PhysicsSolver = SolverComponent->GetPhysicsSolver())
 		{
-			if (FRigPhysicsSimulation* PhysicsSimulation = Cast<FRigPhysicsSimulation>(
-				ControlRig->GetPhysicsSimulation(PhysicsSolverComponentKey)))
-			{
-				const FRigPhysicsSimulation::FSimulationSpaceData& SimulationSpaceData = 
-					PhysicsSimulation->GetSimulationSpaceData();
-
-				LinearVelocity = SimulationSpaceData.LinearVelocity;
-				AngularVelocity = SimulationSpaceData.AngularVelocity;
-				LinearAcceleration = SimulationSpaceData.LinearAcceleration;
-				AngularAcceleration = SimulationSpaceData.AngularAcceleration;
-				Gravity = SimulationSpaceData.Gravity;
-			}
+			const FRigPhysicsSolver::FSimulationSpaceData& SimulationSpaceData = 
+				PhysicsSolver->GetSimulationSpaceData();
+
+			LinearVelocity = SimulationSpaceData.LinearVelocity;
+			AngularVelocity = SimulationSpaceData.AngularVelocity;
+			LinearAcceleration = SimulationSpaceData.LinearAcceleration;
+			AngularAcceleration = SimulationSpaceData.AngularAcceleration;
+			Gravity = SimulationSpaceData.Gravity;
 		}
 	}
 }
diff --git a/Engine/Plugins/Experimental/ControlRigPhysics/Source/ControlRigPhysics/Private/RigPhysicsSolver_Space.cpp b/Engine/Plugins/Experimental/ControlRigPhysics/Source/ControlRigPhysics/Private/RigPhysicsSolver_Space.cpp
new file mode 100644
--- /dev/null
+++ b/Engine/Plugins/Experimental/ControlRigPhysics/Source/ControlRigPhysics/Private/RigPhysicsSolver_Space.cpp
@@ -0,0 +1,319 @@
+// Copyright Epic Games, Inc. All Rights Reserved.
+
+#include "RigPhysicsSolver.h"
+#include "RigPhysicsSolver_Space.inl"
+
+#include "RigPhysicsSolverComponent.h"
+
+#include "Rigs/RigHierarchy.h"
+#include "Rigs/RigHierarchyController.h"
+
+#include "PhysicsControlHelpers.h"
+
+#include "Engine/Engine.h"
+
+#include "Physics/ImmediatePhysics/ImmediatePhysicsActorHandle.h"
+#include "Physics/ImmediatePhysics/ImmediatePhysicsJointHandle.h"
+#include "Physics/ImmediatePhysics/ImmediatePhysicsSimulation.h"
+
+#include "Physics/Experimental/PhysScene_Chaos.h"
+#include "Chaos/ParticleHandle.h"
+#include "Chaos/PBDJointConstraints.h"
+#include "Chaos/ImplicitObject.h"
+#include "Chaos/ShapeInstance.h"
+#include "Chaos/Capsule.h"
+#include "Chaos/PBDJointConstraintUtilities.h"
+
+#include "PhysicsEngine/BodyInstance.h"
+#include "PhysicsEngine/ConstraintInstance.h"
+#include "PhysicsEngine/BodySetup.h"
+
+TAutoConsoleVariable<int> CVarControlRigPhysicsShowSimulationSpaceInfo(
+	TEXT("ControlRig.Physics.ShowSimulationSpaceInfo"), 0,
+	TEXT("Shows information associated with the simulation space used in control rig physics."));
+
+//======================================================================================================================
+void FRigPhysicsSolver::InitSimulationSpace(
+	const FTransform& ComponentTM,
+	const FTransform& BoneRelComponentTM)
+{
+	SimulationSpaceState.ComponentTM = ComponentTM;
+	SimulationSpaceState.BoneRelComponentTM = BoneRelComponentTM;
+}
+
+//======================================================================================================================
+// Danny TODO We need support for double precision Perlin inputs - for now duplicate.
+namespace FPerlinHelpers
+{
+	// random permutation of 256 numbers, repeated 2x
+	static const int32 Permutation[512] = {
+		63, 9, 212, 205, 31, 128, 72, 59, 137, 203, 195, 170, 181, 115, 165, 40, 116, 139, 175, 225, 132, 99, 222, 2, 41, 15, 197, 93, 169, 90, 228, 43, 221, 38, 206, 204, 73, 17, 97, 10, 96, 47, 32, 138, 136, 30, 219,
+		78, 224, 13, 193, 88, 134, 211, 7, 112, 176, 19, 106, 83, 75, 217, 85, 0, 98, 140, 229, 80, 118, 151, 117, 251, 103, 242, 81, 238, 172, 82, 110, 4, 227, 77, 243, 46, 12, 189, 34, 188, 200, 161, 68, 76, 171, 194,
+		57, 48, 247, 233, 51, 105, 5, 23, 42, 50, 216, 45, 239, 148, 249, 84, 70, 125, 108, 241, 62, 66, 64, 240, 173, 185, 250, 49, 6, 37, 26, 21, 244, 60, 223, 255, 16, 145, 27, 109, 58, 102, 142, 253, 120, 149, 160,
+		124, 156, 79, 186, 135, 127, 14, 121, 22, 65, 54, 153, 91, 213, 174, 24, 252, 131, 192, 190, 202, 208, 35, 94, 231, 56, 95, 183, 163, 111, 147, 25, 67, 36, 92, 236, 71, 166, 1, 187, 100, 130, 143, 237, 178, 158,
+		104, 184, 159, 177, 52, 214, 230, 119, 87, 114, 201, 179, 198, 3, 248, 182, 39, 11, 152, 196, 113, 20, 232, 69, 141, 207, 234, 53, 86, 180, 226, 74, 150, 218, 29, 133, 8, 44, 123, 28, 146, 89, 101, 154, 220, 126,
+		155, 122, 210, 168, 254, 162, 129, 33, 18, 209, 61, 191, 199, 157, 245, 55, 164, 167, 215, 246, 144, 107, 235,
+
+		63, 9, 212, 205, 31, 128, 72, 59, 137, 203, 195, 170, 181, 115, 165, 40, 116, 139, 175, 225, 132, 99, 222, 2, 41, 15, 197, 93, 169, 90, 228, 43, 221, 38, 206, 204, 73, 17, 97, 10, 96, 47, 32, 138, 136, 30, 219,
+		78, 224, 13, 193, 88, 134, 211, 7, 112, 176, 19, 106, 83, 75, 217, 85, 0, 98, 140, 229, 80, 118, 151, 117, 251, 103, 242, 81, 238, 172, 82, 110, 4, 227, 77, 243, 46, 12, 189, 34, 188, 200, 161, 68, 76, 171, 194,
+		57, 48, 247, 233, 51, 105, 5, 23, 42, 50, 216, 45, 239, 148, 249, 84, 70, 125, 108, 241, 62, 66, 64, 240, 173, 185, 250, 49, 6, 37, 26, 21, 244, 60, 223, 255, 16, 145, 27, 109, 58, 102, 142, 253, 120, 149, 160,
+		124, 156, 79, 186, 135, 127, 14, 121, 22, 65, 54, 153, 91, 213, 174, 24, 252, 131, 192, 190, 202, 208, 35, 94, 231, 56, 95, 183, 163, 111, 147, 25, 67, 36, 92, 236, 71, 166, 1, 187, 100, 130, 143, 237, 178, 158,
+		104, 184, 159, 177, 52, 214, 230, 119, 87, 114, 201, 179, 198, 3, 248, 182, 39, 11, 152, 196, 113, 20, 232, 69, 141, 207, 234, 53, 86, 180, 226, 74, 150, 218, 29, 133, 8, 44, 123, 28, 146, 89, 101, 154, 220, 126,
+		155, 122, 210, 168, 254, 162, 129, 33, 18, 209, 61, 191, 199, 157, 245, 55, 164, 167, 215, 246, 144, 107, 235
+	};
+
+	// Gradient functions for 1D, 2D and 3D Perlin noise
+
+	FORCEINLINE float Grad1(int32 Hash, double X)
+	{
+		// Slicing Perlin's 3D improved noise would give us only scales of -1, 0 and 1; this looks pretty bad so let's use a different sampling
+		static const double Grad1Scales[16] = { -8 / 8, -7 / 8., -6 / 8., -5 / 8., -4 / 8., -3 / 8., -2 / 8., -1 / 8., 1 / 8., 2 / 8., 3 / 8., 4 / 8., 5 / 8., 6 / 8., 7 / 8., 8 / 8 };
+		return Grad1Scales[Hash & 15] * X;
+	}
+
+
+	// Curve w/ second derivative vanishing at 0 and 1, from Perlin's improved noise paper
+	FORCEINLINE double SmoothCurve(double X)
+	{
+		return X * X * X * (X * (X * 6.0 - 15.0) + 10.0);
+	}
+
+	FORCEINLINE float PerlinNoise1D(double X)
+	{
+		const double Xfl = FMath::FloorToFloat(X);
+		const int64 Xi = (int64)(Xfl) & 255;
+		X -= Xfl;
+		const double Xm1 = X - 1.0;
+
+		const int32 A = Permutation[Xi];
+		const int32 B = Permutation[Xi + 1];
+
+		const double U = SmoothCurve(X);
+
+		// 2.0 factor to ensure (-1, 1) range
+		return 2.0f * FMath::Lerp(Grad1(A, X), Grad1(B, Xm1), U);
+	}
+}
+
+//======================================================================================================================
+// Note - don't use the space conversion functions here as the state won't have been set yet.
+FRigPhysicsSolver::FSimulationSpaceData FRigPhysicsSolver::UpdateSimulationSpaceStateAndCalculateData(
+	const FRigVMExecuteContext&         ExecuteContext, 
+	const URigHierarchy&                Hierarchy,
+	const FRigPhysicsSolverComponent&   SolverComponent,
+	const float                         Dt,
+	const double                        AbsoluteTime)
+{
+	const FRigPhysicsSolverSettings& SolverSettings = SolverComponent.SolverSettings;
+	const FRigPhysicsSimulationSpaceSettings& SimulationSpaceSettings = SolverComponent.SimulationSpaceSettings;
+
+	SimulationSpaceState.ComponentTM = ExecuteContext.GetToWorldSpaceTransform();
+
+	if (SolverSettings.SimulationSpace == ERigPhysicsSimulationSpace::SpaceBone &&  SolverSettings.SpaceBone.IsValid())
+	{
+		SimulationSpaceState.BoneRelComponentTM = Hierarchy.GetGlobalTransform(SolverSettings.SpaceBone);
+	}
+
+	// Record the history - but avoid polluting it with zero Dt updates. What that means is - if we
+	// get a zero-Dt update, then just update our current sim space TM, which means the time delta
+	// from the previous state is actually the current Dt (i.e. don't overwrite the current Dt).
+	if (Dt > SMALL_NUMBER)
+	{
+		SimulationSpaceState.PrevDt = SimulationSpaceState.Dt;
+		SimulationSpaceState.Dt = Dt;
+
+		SimulationSpaceState.PrevPrevSimulationSpaceTM = SimulationSpaceState.PrevSimulationSpaceTM;
+		SimulationSpaceState.PrevSimulationSpaceTM = SimulationSpaceState.SimulationSpaceTM;
+	}
+
+	SimulationSpaceState.SimulationSpaceTM = GetSpaceTransform(
+		SolverSettings.SimulationSpace, SimulationSpaceState.ComponentTM, SimulationSpaceState.BoneRelComponentTM);
+
+	SimulationSpaceData = FSimulationSpaceData();
+	SimulationSpaceData.Gravity = ::ConvertWorldVectorToSimSpaceNoScale(
+		SolverSettings.SimulationSpace, SolverSettings.Gravity, 
+		SimulationSpaceState.ComponentTM, SimulationSpaceState.BoneRelComponentTM);
+
+	if (SolverSettings.SimulationSpace == ERigPhysicsSimulationSpace::World)
+	{
+		// Danny TODO This is probably redundant unless we support runtime switching of the space
+		InitSimulationSpace(SimulationSpaceState.ComponentTM, SimulationSpaceState.BoneRelComponentTM);
+		SimulationSpaceData.LinearVelocity = SimulationSpaceSettings.ExternalLinearVelocity;
+		SimulationSpaceData.AngularVelocity = FMath::DegreesToRadians(SimulationSpaceSettings.ExternalAngularVelocity);
+		return SimulationSpaceData;
+	}
+
+	// If the timestep is zero, or if we have skipped an update, then it doesn't actually matter
+	// what the velocity is - but make sure it doesn't corrupt anything.
+	if (SimulationSpaceState.Dt < SMALL_NUMBER || UpdateCounter != PreviousUpdateCounter + 1)
+	{
+		SimulationSpaceData.LinearVelocity = FVector::ZeroVector;
+		SimulationSpaceData.AngularVelocity = FVector::ZeroVector;
+		SimulationSpaceData.LinearAcceleration = FVector::ZeroVector;
+		SimulationSpaceData.AngularAcceleration = FVector::ZeroVector;
+		return SimulationSpaceData;
+	}
+
+	// We calculate velocities etc in world space first, and then subsequently convert them into
+	// simulation space.
+
+	// Note that the velocity/accel calculations are intended to track the world/simulation behavior
+	// - not necessarily be the most accurate calculations! For example, we could use one-sided
+	// finite difference approximations, but this wouldn't necessarily be correct.
+
+	// World-space component linear velocity and acceleration
+	SimulationSpaceData.LinearVelocity = UE::PhysicsControl::CalculateLinearVelocity(
+		SimulationSpaceState.PrevSimulationSpaceTM.GetTranslation(),
+		SimulationSpaceState.SimulationSpaceTM.GetTranslation(), SimulationSpaceState.Dt);
+	const FVector PrevSpaceLinearVel = 
+		SimulationSpaceState.PrevDt < SMALL_NUMBER 
+		? SimulationSpaceData.LinearVelocity 
+		: UE::PhysicsControl::CalculateLinearVelocity(
+			SimulationSpaceState.PrevPrevSimulationSpaceTM.GetTranslation(),
+			SimulationSpaceState.PrevSimulationSpaceTM.GetTranslation(), SimulationSpaceState.PrevDt);
+	SimulationSpaceData.LinearAcceleration =
+		(SimulationSpaceData.LinearVelocity - PrevSpaceLinearVel) / SimulationSpaceState.Dt;
+
+	// World-space component angular velocity and acceleration
+	SimulationSpaceData.AngularVelocity = UE::PhysicsControl::CalculateAngularVelocity(
+		SimulationSpaceState.PrevSimulationSpaceTM.GetRotation(),
+		SimulationSpaceState.SimulationSpaceTM.GetRotation(), SimulationSpaceState.Dt);
+	const FVector PrevSpaceAngularVel = 
+		SimulationSpaceState.PrevDt < SMALL_NUMBER 
+		? SimulationSpaceData.AngularVelocity
+		: UE::PhysicsControl::CalculateAngularVelocity(
+			SimulationSpaceState.PrevPrevSimulationSpaceTM.GetRotation(),
+			SimulationSpaceState.PrevSimulationSpaceTM.GetRotation(), SimulationSpaceState.PrevDt);
+	SimulationSpaceData.AngularAcceleration =
+		(SimulationSpaceData.AngularVelocity - PrevSpaceAngularVel) / SimulationSpaceState.Dt;
+
+	// Apply Z scale
+	SimulationSpaceData.LinearVelocity.Z *= SimulationSpaceSettings.VelocityScaleZ;
+	SimulationSpaceData.LinearAcceleration.Z *= SimulationSpaceSettings.VelocityScaleZ;
+
+	bool bLinearAccelerationTrigger = SimulationSpaceSettings.LinearAccelerationThresholdForTeleport > 0 &&
+		SimulationSpaceData.LinearAcceleration.SquaredLength() >
+		FMath::Square(SimulationSpaceSettings.LinearAccelerationThresholdForTeleport);
+	bool bAngularAccelerationTrigger = SimulationSpaceSettings.AngularAccelerationThresholdForTeleport > 0 &&
+		SimulationSpaceData.AngularAcceleration.SquaredLength() >
+		FMath::Square(FMath::DegreesToRadians(SimulationSpaceSettings.AngularAccelerationThresholdForTeleport));
+	bool bPositionTrigger = SimulationSpaceSettings.PositionChangeThresholdForTeleport > 0 &&
+		SimulationSpaceData.LinearVelocity.SquaredLength() >
+		FMath::Square(SimulationSpaceSettings.PositionChangeThresholdForTeleport / SimulationSpaceState.Dt);
+	bool bOrientationTrigger = SimulationSpaceSettings.OrientationChangeThresholdForTeleport > 0 &&
+		SimulationSpaceData.AngularVelocity.SquaredLength() >
+		FMath::Square(FMath::DegreesToRadians(
+			SimulationSpaceSettings.OrientationChangeThresholdForTeleport / SimulationSpaceState.Dt));
+
+	// Check and report teleports
+	if (bLinearAccelerationTrigger || bAngularAccelerationTrigger || bPositionTrigger || bOrientationTrigger)
+	{
+		if (bLinearAccelerationTrigger)
+		{
+			UE_LOG(LogRigPhysics, Log, TEXT("Detected linear Acceleration (%f > %f) teleport"),
+				SimulationSpaceData.LinearAcceleration.Length(), 
+				SimulationSpaceSettings.LinearAccelerationThresholdForTeleport);
+		}
+		if (bAngularAccelerationTrigger)
+		{
+			UE_LOG(LogRigPhysics, Log, TEXT("Detected angular Acceleration (%f > %f) teleport"),
+				FMath::RadiansToDegrees(SimulationSpaceData.AngularAcceleration.Length()),
+				SimulationSpaceSettings.AngularAccelerationThresholdForTeleport);
+		}
+		if (bPositionTrigger)
+		{
+			UE_LOG(LogRigPhysics, Log, TEXT("Detected position (%f > %f) teleport"),
+				SimulationSpaceData.LinearVelocity.Length() * SimulationSpaceState.Dt,
+				SimulationSpaceSettings.PositionChangeThresholdForTeleport);
+		}
+		if (bOrientationTrigger)
+		{
+			UE_LOG(LogRigPhysics, Log, TEXT("Detected orientation (%f > %f) teleport"),
+				FMath::RadiansToDegrees(SimulationSpaceData.AngularVelocity.Length() * SimulationSpaceState.Dt),
+				SimulationSpaceSettings.OrientationChangeThresholdForTeleport);
+		}
+
+		// Note that a teleport detection shouldn't change the pose, or the current motion. We just
+		// don't want to bring in that unwanted global motion.
+		SimulationSpaceData.LinearVelocity = FVector::ZeroVector;
+		SimulationSpaceData.AngularVelocity = FVector::ZeroVector;
+		SimulationSpaceData.LinearAcceleration = FVector::ZeroVector;
+		SimulationSpaceData.AngularAcceleration = FVector::ZeroVector;
+
+		// This will stop the next step from using bogus values too.
+		SimulationSpaceState.PrevSimulationSpaceTM = SimulationSpaceState.SimulationSpaceTM;
+		SimulationSpaceState.PrevPrevSimulationSpaceTM = SimulationSpaceState.SimulationSpaceTM;
+		SimulationSpaceState.PrevDt = 0;
+		SimulationSpaceState.Dt = 0;
+		// Avoid cached transforms being used in controls by bumping the update counter. 
+		UpdateCounter += 1;
+	}
+	else
+	{
+		if (SimulationSpaceSettings.bClampLinearVelocity)
+		{
+			SimulationSpaceData.LinearVelocity =
+				SimulationSpaceData.LinearVelocity.GetClampedToMaxSize(SimulationSpaceSettings.MaxLinearVelocity);
+		}
+		if (SimulationSpaceSettings.bClampAngularVelocity)
+		{
+			SimulationSpaceData.AngularVelocity =
+				SimulationSpaceData.AngularVelocity.GetClampedToMaxSize(SimulationSpaceSettings.MaxAngularVelocity);
+		}
+		if (SimulationSpaceSettings.bClampLinearAcceleration)
+		{
+			SimulationSpaceData.LinearAcceleration =
+				SimulationSpaceData.LinearAcceleration.GetClampedToMaxSize(SimulationSpaceSettings.MaxLinearAcceleration);
+		}
+		if (SimulationSpaceSettings.bClampAngularAcceleration)
+		{
+			SimulationSpaceData.AngularAcceleration =
+				SimulationSpaceData.AngularAcceleration.GetClampedToMaxSize(SimulationSpaceSettings.MaxAngularAcceleration);
+		}
+	}
+
+	SimulationSpaceData.LinearVelocity += SimulationSpaceSettings.ExternalLinearVelocity;
+	SimulationSpaceData.AngularVelocity += FMath::DegreesToRadians(SimulationSpaceSettings.ExternalAngularVelocity);
+
+	if (!SimulationSpaceSettings.ExternalTurbulenceVelocity.IsNearlyZero())
+	{
+		FVector T(
+			FPerlinHelpers::PerlinNoise1D(AbsoluteTime),
+			FPerlinHelpers::PerlinNoise1D(AbsoluteTime + 10.0),
+			FPerlinHelpers::PerlinNoise1D(AbsoluteTime + 20.0));
+
+		FVector Turbulence = T * SimulationSpaceSettings.ExternalTurbulenceVelocity;
+
+		SimulationSpaceData.LinearVelocity += Turbulence;
+
+		if (GEngine && CVarControlRigPhysicsShowSimulationSpaceInfo.GetValueOnAnyThread())
+		{
+			GEngine->AddOnScreenDebugMessage(INDEX_NONE, 0.f, FColor::Yellow,
+				FString::Printf(TEXT("Turbulence %s"), *Turbulence.ToString()));
+		}
+	}
+
+	// Transform world-space motion into simulation space Danny TODO note that this matches the code
+	// in RBAN, and is doing what the interface requires (i.e. movement of the space in the space of
+	// the space!). Danny TODO note that we don't currently support scaling
+	SimulationSpaceData.LinearVelocity =
+		SimulationSpaceState.SimulationSpaceTM.InverseTransformVectorNoScale(SimulationSpaceData.LinearVelocity);
+	SimulationSpaceData.AngularVelocity =
+		SimulationSpaceState.SimulationSpaceTM.InverseTransformVectorNoScale(SimulationSpaceData.AngularVelocity);
+	SimulationSpaceData.LinearAcceleration =
+		SimulationSpaceState.SimulationSpaceTM.InverseTransformVectorNoScale(SimulationSpaceData.LinearAcceleration);
+	SimulationSpaceData.AngularAcceleration =
+		SimulationSpaceState.SimulationSpaceTM.InverseTransformVectorNoScale(SimulationSpaceData.AngularAcceleration);
+
+	if (GEngine && CVarControlRigPhysicsShowSimulationSpaceInfo.GetValueOnAnyThread())
+	{
+		GEngine->AddOnScreenDebugMessage(INDEX_NONE, 0.f, FColor::Yellow,
+			FString::Printf(TEXT("Dt %6.2fms Space P %35s V %35s A %35s"),
+				Dt * 1000.0f,
+				*SimulationSpaceState.SimulationSpaceTM.GetLocation().ToString(),
+				*SimulationSpaceData.LinearVelocity.ToString(),
+				*SimulationSpaceData.LinearAcceleration.ToString()));
+	}
+
+	return SimulationSpaceData;
+}
diff --git a/Engine/Plugins/Experimental/ControlRigPhysics/Source/ControlRigPhysics/Private/RigPhysicsSolver_Space.inl b/Engine/Plugins/Experimental/ControlRigPhysics/Source/ControlRigPhysics/Private/RigPhysicsSolver_Space.inl
new file mode 100644
--- /dev/null
+++ b/Engine/Plugins/Experimental/ControlRigPhysics/Source/ControlRigPhysics/Private/RigPhysicsSolver_Space.inl
@@ -0,0 +1,186 @@
+// Copyright Epic Games, Inc. All Rights Reserved.
+
+#pragma once
+
+#include "RigPhysicsSolver.h"
+#include "PhysicsControlHelpers.h"
+
+#include "Engine/Engine.h"
+
+//======================================================================================================================
+FORCEINLINE FTransform FRigPhysicsSolver::GetSpaceTransform(
+	ERigPhysicsSimulationSpace Space, const FTransform& ComponentTM, const FTransform& BoneTM)
+{
+	switch (Space)
+	{
+	case ERigPhysicsSimulationSpace::Component: return ComponentTM;
+	case ERigPhysicsSimulationSpace::World: return FTransform::Identity;
+	case ERigPhysicsSimulationSpace::SpaceBone: return BoneTM * ComponentTM;
+	default: ensureMsgf(false, TEXT("Unsupported Simulation Space")); return FTransform::Identity;
+	}
+}
+
+//======================================================================================================================
+FORCEINLINE FTransform FRigPhysicsSolver::GetSimulationSpaceTransform(const FRigPhysicsSolverSettings& SolverSettings) const
+{
+	switch (SolverSettings.SimulationSpace)
+	{
+	case ERigPhysicsSimulationSpace::Component: return SimulationSpaceState.ComponentTM;
+	case ERigPhysicsSimulationSpace::World: return FTransform::Identity;
+	case ERigPhysicsSimulationSpace::SpaceBone: return 
+		SimulationSpaceState.BoneRelComponentTM * SimulationSpaceState.ComponentTM;
+	default: ensureMsgf(false, TEXT("Unsupported Simulation Space")); return FTransform::Identity;
+	}
+}
+
+//======================================================================================================================
+FORCEINLINE FTransform FRigPhysicsSolver::ConvertComponentSpaceTransformToSimSpace(
+	const FRigPhysicsSolverSettings& SolverSettings, const FTransform& TM) const
+{
+	switch (SolverSettings.SimulationSpace)
+	{
+	case ERigPhysicsSimulationSpace::Component: return TM;
+	case ERigPhysicsSimulationSpace::World: return TM * SimulationSpaceState.ComponentTM;
+	case ERigPhysicsSimulationSpace::SpaceBone: return TM.GetRelativeTransform(SimulationSpaceState.BoneRelComponentTM);
+	default: ensureMsgf(false, TEXT("Unsupported Simulation Space")); return TM;
+	}
+}
+
+//======================================================================================================================
+FORCEINLINE FVector FRigPhysicsSolver::ConvertComponentSpacePositionToSimSpace(
+	const FRigPhysicsSolverSettings& SolverSettings, const FVector& P) const
+{
+	switch (SolverSettings.SimulationSpace)
+	{
+	case ERigPhysicsSimulationSpace::Component: return P;
+	case ERigPhysicsSimulationSpace::World:  return SimulationSpaceState.ComponentTM.TransformPositionNoScale(P);
+	case ERigPhysicsSimulationSpace::SpaceBone: return SimulationSpaceState.BoneRelComponentTM.InverseTransformPositionNoScale(P);
+	default: ensureMsgf(false, TEXT("Unsupported Simulation Space")); return P;
+	}
+}
+
+//======================================================================================================================
+FORCEINLINE FVector FRigPhysicsSolver::ConvertComponentSpaceVectorToSimSpace(
+	const FRigPhysicsSolverSettings& SolverSettings, const FVector& V) const
+{
+	switch (SolverSettings.SimulationSpace)
+	{
+	case ERigPhysicsSimulationSpace::Component: return V;
+	case ERigPhysicsSimulationSpace::World: return SimulationSpaceState.ComponentTM.TransformVectorNoScale(V);
+	case ERigPhysicsSimulationSpace::SpaceBone: 
+		return SimulationSpaceState.BoneRelComponentTM.InverseTransformVectorNoScale(V);
+	default: ensureMsgf(false, TEXT("Unsupported Simulation Space")); return V;
+	}
+}
+
+//======================================================================================================================
+FORCEINLINE FTransform FRigPhysicsSolver::ConvertSimSpaceTransformToComponentSpace(
+	const FRigPhysicsSolverSettings& SolverSettings, const FTransform& TM) const
+{
+	switch (SolverSettings.SimulationSpace)
+	{
+	case ERigPhysicsSimulationSpace::Component: return TM;
+	case ERigPhysicsSimulationSpace::World: return TM.GetRelativeTransform(SimulationSpaceState.ComponentTM);
+	case ERigPhysicsSimulationSpace::SpaceBone: return TM * SimulationSpaceState.BoneRelComponentTM;
+	default: ensureMsgf(false, TEXT("Unsupported Simulation Space")); return TM;
+	}
+}
+
+//======================================================================================================================
+FORCEINLINE FVector FRigPhysicsSolver::ConvertSimSpacePositionToComponentSpace(
+	const FRigPhysicsSolverSettings& SolverSettings, const FVector& Position) const
+{
+	switch (SolverSettings.SimulationSpace)
+	{
+	case ERigPhysicsSimulationSpace::Component: return Position;
+	case ERigPhysicsSimulationSpace::World: return SimulationSpaceState.ComponentTM.InverseTransformPositionNoScale(Position);
+	case ERigPhysicsSimulationSpace::SpaceBone: return SimulationSpaceState.BoneRelComponentTM.TransformPositionNoScale(Position);
+	default: ensureMsgf(false, TEXT("Unsupported Simulation Space")); return Position;
+	}
+}
+
+//======================================================================================================================
+FORCEINLINE FVector FRigPhysicsSolver::ConvertSimSpaceVectorToComponentSpace(
+	const FRigPhysicsSolverSettings& SolverSettings, const FVector& Vector) const
+{
+	switch (SolverSettings.SimulationSpace)
+	{
+	case ERigPhysicsSimulationSpace::Component: return Vector;
+	case ERigPhysicsSimulationSpace::World: return SimulationSpaceState.ComponentTM.InverseTransformVectorNoScale(Vector);
+	case ERigPhysicsSimulationSpace::SpaceBone: return SimulationSpaceState.BoneRelComponentTM.TransformVectorNoScale(Vector);
+	default: ensureMsgf(false, TEXT("Unsupported Simulation Space")); return Vector;
+	}
+}
+
+//======================================================================================================================
+FORCEINLINE FVector ConvertWorldVectorToSimSpaceNoScale(
+	ERigPhysicsSimulationSpace Space, const FVector& WorldVector, 
+	const FTransform& ComponentTM, const FTransform& BoneTM)
+{
+	switch (Space)
+	{
+	case ERigPhysicsSimulationSpace::Component: return ComponentTM.InverseTransformVectorNoScale(WorldVector);
+	case ERigPhysicsSimulationSpace::World: return WorldVector;
+	case ERigPhysicsSimulationSpace::SpaceBone:
+		return BoneTM.InverseTransformVectorNoScale(ComponentTM.InverseTransformVectorNoScale(WorldVector));
+	default: ensureMsgf(false, TEXT("Unsupported Simulation Space")); return WorldVector;
+	}
+}
+
+//======================================================================================================================
+FORCEINLINE FVector FRigPhysicsSolver::ConvertWorldVectorToSimSpaceNoScale(
+	const FRigPhysicsSolverSettings& SolverSettings, const FVector& WorldVector) const
+{
+	switch (SolverSettings.SimulationSpace)
+	{
+	case ERigPhysicsSimulationSpace::Component: 
+		return SimulationSpaceState.ComponentTM.InverseTransformVectorNoScale(WorldVector);
+	case ERigPhysicsSimulationSpace::World: return WorldVector;
+	case ERigPhysicsSimulationSpace::SpaceBone:
+		return SimulationSpaceState.BoneRelComponentTM.InverseTransformVectorNoScale(
+			SimulationSpaceState.ComponentTM.InverseTransformVectorNoScale(WorldVector));
+	default: ensureMsgf(false, TEXT("Unsupported Simulation Space")); return WorldVector;
+	}
+}
+
+//======================================================================================================================
+FORCEINLINE FVector FRigPhysicsSolver::ConvertWorldPositionToSimSpaceNoScale(
+	const FRigPhysicsSolverSettings& SolverSettings, const FVector& WorldPosition) const
+{
+	switch (SolverSettings.SimulationSpace)
+	{
+	case ERigPhysicsSimulationSpace::Component: 
+		return SimulationSpaceState.ComponentTM.InverseTransformPositionNoScale(WorldPosition);
+	case ERigPhysicsSimulationSpace::World: return WorldPosition;
+	case ERigPhysicsSimulationSpace::SpaceBone:
+		return SimulationSpaceState.BoneRelComponentTM.InverseTransformPositionNoScale(
+			SimulationSpaceState.ComponentTM.InverseTransformPositionNoScale(WorldPosition));
+	default: ensureMsgf(false, TEXT("Unsupported Simulation Space")); return WorldPosition;
+	}
+}
+
+//======================================================================================================================
+FORCEINLINE FTransform FRigPhysicsSolver::ConvertWorldTransformToSimSpace(
+	const FRigPhysicsSolverSettings& SolverSettings, const FTransform& WorldTM) const
+{
+	switch (SolverSettings.SimulationSpace)
+	{
+	case ERigPhysicsSimulationSpace::Component: return WorldTM.GetRelativeTransform(SimulationSpaceState.ComponentTM);
+	case ERigPhysicsSimulationSpace::World: return WorldTM;
+	case ERigPhysicsSimulationSpace::SpaceBone:
+		return WorldTM.GetRelativeTransform(SimulationSpaceState.BoneRelComponentTM * SimulationSpaceState.ComponentTM);
+	default: ensureMsgf(false, TEXT("Unsupported Simulation Space")); return WorldTM;
+	}
+}
+
+//======================================================================================================================
+FORCEINLINE FTransform FRigPhysicsSolver::ConvertCollisionSpaceTransformToSimSpace(
+	const FRigPhysicsSolverSettings& SolverSettings, const FTransform& TM) const
+{
+	FTransform SimSpaceTM = GetSpaceTransform(
+		SolverSettings.SimulationSpace, SimulationSpaceState.ComponentTM, SimulationSpaceState.BoneRelComponentTM);
+	FTransform CollisionSpaceTM = GetSpaceTransform(
+		SolverSettings.CollisionSpace, SimulationSpaceState.ComponentTM, SimulationSpaceState.BoneRelComponentTM);
+	FTransform WorldSpaceTM = TM * CollisionSpaceTM;
+	return WorldSpaceTM.GetRelativeTransform(SimSpaceTM);
+}
diff --git a/Engine/Plugins/Experimental/ControlRigPhysics/Source/ControlRigPhysics/Private/RigPhysicsSolver_Visualization.cpp b/Engine/Plugins/Experimental/ControlRigPhysics/Source/ControlRigPhysics/Private/RigPhysicsSolver_Visualization.cpp
new file mode 100644
--- /dev/null
+++ b/Engine/Plugins/Experimental/ControlRigPhysics/Source/ControlRigPhysics/Private/RigPhysicsSolver_Visualization.cpp
@@ -0,0 +1,529 @@
+// Copyright Epic Games, Inc. All Rights Reserved.
+
+#include "RigPhysicsSolver.h"
+#include "RigPhysicsSolver_Space.inl"
+
+#include "Rigs/RigHierarchy.h"
+#include "Rigs/RigHierarchyController.h"
+
+#include "RigVMCore/RigVMDrawInterface.h"
+
+#include "Physics/ImmediatePhysics/ImmediatePhysicsActorHandle.h"
+#include "Physics/ImmediatePhysics/ImmediatePhysicsJointHandle.h"
+#include "Physics/ImmediatePhysics/ImmediatePhysicsSimulation.h"
+
+#include "Physics/Experimental/PhysScene_Chaos.h"
+#include "Chaos/ParticleHandle.h"
+#include "Chaos/PBDJointConstraints.h"
+#include "Chaos/ImplicitObject.h"
+#include "Chaos/ShapeInstance.h"
+#include "Chaos/Capsule.h"
+#include "Chaos/PBDJointConstraintUtilities.h"
+
+#include "Engine/Engine.h"
+#include "Materials/MaterialInstanceDynamic.h"
+
+#include "Math/Color.h"
+
+// Danny TODO expose these in the visualization settings
+namespace RigPhysicsSolverDraw
+{
+	FLinearColor DynamicCollisionColor = FColor::Yellow;
+	FLinearColor DynamicNoCollisionColor = FColor::Orange;
+	FLinearColor KinematicCollisionColor = FColor::Blue;
+	FLinearColor KinematicNoCollisionColor = FColor::Purple;
+
+	FLinearColor ActiveContactColor = FColor::Red;
+	FLinearColor InactiveContactColor = FColor::Silver;
+}
+
+TAutoConsoleVariable<int> CVarControlRigPhysicsShowActiveContactsOverride(
+	TEXT("ControlRig.Physics.ShowActiveContactsOveride"), -1,
+	TEXT("Whether to draw active contacts (requires visualization to be enabled). -1 uses the visualization setting, 0 forces drawing to be disabled, 1 forces it to be enabled."));
+
+TAutoConsoleVariable<int> CVarControlRigPhysicsShowInactiveContactsOverride(
+	TEXT("ControlRig.Physics.ShowInactiveContactsOveride"), -1,
+	TEXT("Whether to draw inactive contacts (requires visualization to be enabled). -1 uses the visualization setting, 0 forces drawing to be disabled, 1 forces it to be enabled."));
+
+//======================================================================================================================
+static FLinearColor GetColor(bool bIsKinematic, bool bCollides)
+{
+	if (bIsKinematic)
+	{
+		return bCollides 
+			? RigPhysicsSolverDraw::KinematicCollisionColor : RigPhysicsSolverDraw::KinematicNoCollisionColor;
+	}
+	else
+	{
+		return bCollides 
+			? RigPhysicsSolverDraw::DynamicCollisionColor : RigPhysicsSolverDraw::DynamicNoCollisionColor;
+	}
+}
+
+//======================================================================================================================
+static void DrawShapes(
+	FRigVMDrawInterface&                    DI,
+	const FRigPhysicsVisualizationSettings& VisualizationSettings,
+	const bool                              bIsKinematic,
+	const bool                              bCollides,
+	const Chaos::FGeometryParticleHandle*   Particle,
+	const FTransform&                       ShapeTransform, 
+	const Chaos::FImplicitObject*           ImplicitObject,
+	const Chaos::FPerShapeData*             Shape)
+{
+	using namespace Chaos;
+
+	if (!ImplicitObject)
+	{
+		return;
+	}
+	const EImplicitObjectType PackedType = ImplicitObject->GetType(); // Type includes scaling and instancing data
+	const EImplicitObjectType InnerType = GetInnerType(ImplicitObject->GetType());
+
+	// For simplicity, we're going to assume no scaling, instancing, transforming etc
+
+	if (!ImplicitObject)
+	{
+		return;
+	}
+
+	switch (InnerType)
+	{
+	case ImplicitObjectType::Transformed:
+	{
+		const TImplicitObjectTransformed<FReal, 3>* Transformed = 
+			ImplicitObject->GetObject<TImplicitObjectTransformed<FReal, 3>>();
+		FTransform TransformedTransform = FTransform(
+			ShapeTransform.GetRotation() * Transformed->GetTransform().GetRotation(),
+			ShapeTransform.TransformPosition(Transformed->GetTransform().GetLocation()));
+		DrawShapes(
+			DI, VisualizationSettings, bIsKinematic, bCollides, Particle, 
+			TransformedTransform, Transformed->GetTransformedObject(), Shape);
+		return;
+	}
+	case ImplicitObjectType::Sphere:
+	{
+		const Chaos::FSphere* Sphere = ImplicitObject->GetObject<Chaos::FSphere>();
+		DI.DrawSphere(ShapeTransform, FTransform(FVec3(Sphere->GetCenterOfMass())), Sphere->GetRadiusf(),
+			GetColor(bIsKinematic, bCollides),
+			VisualizationSettings.LineThickness, VisualizationSettings.ShapeDetail);
+		break;
+	}
+	case ImplicitObjectType::Box:
+	{
+		const TBox<FReal, 3>* Box = ImplicitObject->GetObject<TBox<FReal, 3>>(); // Axis aligned
+		DI.DrawBox(
+			FTransform(FQuat::Identity, Box->GetCenter()) * ShapeTransform, 
+			FTransform(FQuat::Identity, FVector::ZeroVector, Box->Extents()),
+			GetColor(bIsKinematic, bCollides),
+			VisualizationSettings.LineThickness);
+		break;
+	}
+	case ImplicitObjectType::Capsule:
+	{
+		const Chaos::FCapsule* Capsule = ImplicitObject->GetObject<Chaos::FCapsule>();
+		const FQuat Q = FRotationMatrix::MakeFromZ(Capsule->GetAxis()).ToQuat();
+		DI.DrawCapsule(ShapeTransform, FTransform(Q, Capsule->GetCenterOfMass()),
+			Capsule->GetRadiusf(), Capsule->GetHeightf(),
+			GetColor(bIsKinematic, bCollides),
+			VisualizationSettings.LineThickness);
+		break;
+	}
+	default:
+		// If we don't know what it is, don't draw it.
+		break;
+	}
+}
+
+//======================================================================================================================
+static void DrawActor(
+	FRigVMDrawInterface&                    DI,
+	const FRigPhysicsVisualizationSettings& VisualizationSettings,
+	const FTransform&                       SpaceTransform,
+	const ImmediatePhysics::FActorHandle&   ActorHandle)
+{
+	bool bIsKinematic = ActorHandle.GetIsKinematic();
+	bool bCollides = ActorHandle.GetHasCollision();
+
+	const Chaos::FGeometryParticleHandle* GeometryParticleHandle = ActorHandle.GetParticle();
+	const FTransform ParticleTransform = GeometryParticleHandle->GetTransformXR() * SpaceTransform;
+	for (const Chaos::FShapeInstancePtr& ShapeInstance : GeometryParticleHandle->ShapeInstances())
+	{
+		const Chaos::FImplicitObject* ImplicitObject = ShapeInstance->GetGeometry();
+		DrawShapes(
+			DI, VisualizationSettings, bIsKinematic, bCollides, GeometryParticleHandle,
+			ParticleTransform, ImplicitObject, ShapeInstance.Get());
+	}
+
+	if (VisualizationSettings.bShowCentreOfMass)
+	{
+		float Size = 1.0f * VisualizationSettings.ShapeSize;
+		DI.DrawAxes(
+			ParticleTransform, ActorHandle.GetLocalCoMTransform(),
+			Size, VisualizationSettings.LineThickness);
+	}
+}
+
+//======================================================================================================================
+enum
+{
+	TwistIndex = 0,
+	Swing1Index = 1,
+	Swing2Index = 2
+};
+
+//======================================================================================================================
+// Returns the limit range - it will be the value passed in unless it the motion is free or locked. 
+float GetLimitAngleRadians(float LimitAngle, Chaos::EJointMotionType LimitType)
+{
+	switch (LimitType)
+	{
+	case Chaos::EJointMotionType::Free: return UE_PI;
+	case Chaos::EJointMotionType::Locked: return 0.f;
+	default: return LimitAngle;
+	}
+}
+
+//======================================================================================================================
+static void DrawJoint(
+	FRigVMDrawInterface&                    DI,
+	const FRigPhysicsVisualizationSettings& VisualizationSettings,
+	const FTransform&                       SpaceTransform,
+	const ImmediatePhysics::FJointHandle&   JointHandle,
+	const bool                              bDrawAsAxes)
+{
+	const Chaos::FPBDJointConstraintHandle* ConstraintHandle = JointHandle.GetConstraint();
+	if (!ConstraintHandle)
+	{
+		return;
+	}
+
+	const Chaos::TVec2<const ImmediatePhysics_Chaos::FActorHandle*>& ActorHandles = JointHandle.GetActorHandles();
+	const ImmediatePhysics_Chaos::FActorHandle* ChildActor = ActorHandles[0];
+	const ImmediatePhysics_Chaos::FActorHandle* ParentActor = ActorHandles[1];
+
+	FTransform ChildActorTM = (ChildActor ? ChildActor->GetWorldTransform() : FTransform()) * SpaceTransform;
+	FTransform ParentActorTM = (ParentActor ? ParentActor->GetWorldTransform() : FTransform()) * SpaceTransform;
+
+	bool bChildIsKinematic = ChildActor ? ChildActor->GetIsKinematic() : true;
+	bool bParentIsKinematic = ParentActor ? ParentActor->GetIsKinematic() : true;
+
+	float Size = 5.0f * VisualizationSettings.ShapeSize;
+
+	const Chaos::FPBDJointSettings& JointSettings = ConstraintHandle->GetSettings();
+
+	FTransform DialFrame = JointSettings.ConnectorTransforms[0] * ChildActorTM;
+	FTransform LimitFrame = JointSettings.ConnectorTransforms[1] * ParentActorTM;
+
+	if (bDrawAsAxes)
+	{
+		DI.DrawAxes(
+			FTransform(JointSettings.AngularDrivePositionTarget, JointSettings.LinearDrivePositionTarget) * LimitFrame, 
+			FTransform(), Size, VisualizationSettings.LineThickness * 2);
+		DI.DrawAxes(DialFrame, FTransform(), Size, VisualizationSettings.LineThickness);
+	}
+
+	{
+		// See FConstraintInstance::DrawConstraintImp for inspiration
+
+		// There seems to be a swap between swing1 and swing2 compared to GetSwingTwistAngles
+		FVector LimitAngleRadians(
+			GetLimitAngleRadians(JointSettings.AngularLimits[0], JointSettings.AngularMotionTypes[0]),
+			GetLimitAngleRadians(JointSettings.AngularLimits[2], JointSettings.AngularMotionTypes[2]),
+			GetLimitAngleRadians(JointSettings.AngularLimits[1], JointSettings.AngularMotionTypes[1])
+		);
+
+		Chaos::FReal TwistAngle, Swing1Angle, Swing2Angle;
+		const FQuat LimitQ = LimitFrame.GetRotation();
+		FQuat DialQ = DialFrame.GetRotation();
+		DialQ.EnforceShortestArcWith(LimitQ);
+		Chaos::FPBDJointUtilities::GetSwingTwistAngles(LimitQ, DialQ, TwistAngle, Swing1Angle, Swing2Angle);
+
+		bool bDrawViolatedLimits = true;
+
+		bool bTwistViolated = false;
+		bool bSwing1Violated = false;
+		bool bSwing2Violated = false;
+
+		if (bDrawViolatedLimits)
+		{
+			bTwistViolated = JointSettings.AngularMotionTypes[TwistIndex] == Chaos::EJointMotionType::Limited &&
+				FMath::Abs(TwistAngle) > LimitAngleRadians[0];
+
+			bSwing1Violated = JointSettings.AngularMotionTypes[Swing1Index] == Chaos::EJointMotionType::Limited &&
+				FMath::Abs(Swing1Angle) > LimitAngleRadians[1];
+
+			bSwing2Violated = JointSettings.AngularMotionTypes[Swing2Index] == Chaos::EJointMotionType::Limited &&
+				FMath::Abs(Swing2Angle) > LimitAngleRadians[2];
+		}
+
+		const bool bLockSwing1 = JointSettings.AngularMotionTypes[Swing1Index] == Chaos::EJointMotionType::Locked;
+		const bool bLockSwing2 = JointSettings.AngularMotionTypes[Swing2Index] == Chaos::EJointMotionType::Locked;
+		const bool bLockAllSwing = bLockSwing1 && bLockSwing2;
+
+		check(GEngine->ConstraintLimitMaterialX && GEngine->ConstraintLimitMaterialY && GEngine->ConstraintLimitMaterialZ);
+		static UMaterialInterface* LimitMaterialX = GEngine->ConstraintLimitMaterialX;
+		static UMaterialInterface* LimitMaterialXAxis = GEngine->ConstraintLimitMaterialXAxis;
+		static UMaterialInterface* LimitMaterialY = GEngine->ConstraintLimitMaterialY;
+		static UMaterialInterface* LimitMaterialYAxis = GEngine->ConstraintLimitMaterialYAxis;
+		static UMaterialInterface* LimitMaterialZ = GEngine->ConstraintLimitMaterialZ;
+		static UMaterialInterface* LimitMaterialZAxis = GEngine->ConstraintLimitMaterialZAxis;
+
+		float ArrowLength = Size * 1.05f; // stick out a little bit 
+
+		// If swing is limited (but not locked) - draw the swing limit cone.
+		if (!bLockAllSwing)
+		{
+			if (JointSettings.AngularMotionTypes[Swing1Index] == Chaos::EJointMotionType::Free && 
+				JointSettings.AngularMotionTypes[Swing2Index] == Chaos::EJointMotionType::Free)
+			{
+				DI.DrawSphere(LimitFrame, FTransform(), Size * 0.2f, FColor::White, 
+					VisualizationSettings.LineThickness, VisualizationSettings.ShapeDetail);
+			}
+			else
+			{
+				FTransform ConeLimitTM = LimitFrame;
+				ConeLimitTM.SetScale3D(FVector(Size));
+				const float Swing1LimitAngle = LimitAngleRadians[Swing1Index];
+				const float Swing2LimitAngle = LimitAngleRadians[Swing2Index];
+				DI.DrawCone(
+					ConeLimitTM, FTransform(), Swing1LimitAngle, Swing2LimitAngle, VisualizationSettings.ShapeDetail, 
+					true, FColor::Red, LimitMaterialX->GetRenderProxy(), VisualizationSettings.LineThickness);
+			}
+			// Draw the swing Dial indicator - shows the current orientation of the child frame
+			// relative to the parent frame on the swing axis.
+			// Start the arrow at the limit position as it can be confusing if there is joint separation
+			DI.DrawArrow(FTransform(LimitFrame.GetTranslation()), 
+				DialFrame.GetUnitAxis(EAxis::X) * ArrowLength,
+				DialFrame.GetUnitAxis(EAxis::Y), FColor::Red, VisualizationSettings.LineThickness);
+
+			if (bSwing1Violated || bSwing2Violated)
+			{
+				// Drawing a sphere seems a bit heavy... but it appears that if DrawPoint is used then they get
+				// culled when you get close!
+				DI.DrawSphere(FTransform(
+					DialFrame.GetRotation(),
+					LimitFrame.GetTranslation() + DialFrame.GetUnitAxis(EAxis::X) * ArrowLength),
+					FTransform(),
+					ArrowLength * 0.01f, FColor::Orange, VisualizationSettings.LineThickness * 4);
+			}
+
+		}
+
+		// Draw the twist limit
+		if (JointSettings.AngularMotionTypes[TwistIndex] != Chaos::EJointMotionType::Locked)
+		{
+			// Draw as a flat cone
+			FTransform ConeLimitTM = LimitFrame;
+			ConeLimitTM.SetScale3D(FVector(Size));
+			const float TwistLimitAngle = LimitAngleRadians[TwistIndex];
+			DI.DrawCone(ConeLimitTM, FTransform(
+				FQuat::MakeFromRotationVector(FVector(0, 1, 0) * -HALF_PI) * 
+				FQuat::MakeFromRotationVector(FVector(0, 0, 1) * HALF_PI)),
+				TwistLimitAngle, 0.0f, VisualizationSettings.ShapeDetail, true, FColor::Green,
+				LimitMaterialYAxis->GetRenderProxy(), VisualizationSettings.LineThickness);
+
+			// Draw the twist Dial indicator - shows the current orientation of the child frame
+			// relative to the parent frame on the twist axis.
+			FQuat Rot(LimitFrame.GetUnitAxis(EAxis::X), TwistAngle);
+			FVector TwistArrow = Rot * LimitFrame.GetUnitAxis(EAxis::Y);
+			FVector TwistSideDir = FVector::CrossProduct(LimitFrame.GetUnitAxis(EAxis::Z), TwistArrow);
+
+			DI.DrawArrow(FTransform(LimitFrame.GetTranslation()),
+				TwistArrow * ArrowLength,
+				TwistSideDir, FColor::Green, VisualizationSettings.LineThickness);
+
+			if (bTwistViolated)
+			{
+				DI.DrawSphere(FTransform(
+					DialFrame.GetRotation(),
+					LimitFrame.GetTranslation() + TwistArrow * ArrowLength),
+					FTransform(),
+					ArrowLength * 0.01f, FColor::Orange, VisualizationSettings.LineThickness * 4);
+			}
+		}
+	}
+}
+
+//======================================================================================================================
+static bool IsConstraintDriveActive(const Chaos::FPBDJointConstraintHandle& Constraint)
+{
+	const Chaos::FPBDJointSettings& Settings = Constraint.GetSettings();
+	if (Settings.bAngularSLerpPositionDriveEnabled || Settings.bAngularSwingPositionDriveEnabled || Settings.bAngularTwistPositionDriveEnabled)
+	{
+		if (Settings.AngularDriveStiffness.SizeSquared() > 0.0)
+		{
+			return true;
+		}
+	}
+	if (Settings.bAngularSLerpVelocityDriveEnabled || Settings.bAngularSwingVelocityDriveEnabled || Settings.bAngularTwistVelocityDriveEnabled)
+	{
+		if (Settings.AngularDriveDamping.SizeSquared() > 0.0)
+		{
+			return true;
+		}
+	}
+	if (Settings.bLinearPositionDriveEnabled.Max() && Settings.LinearDriveStiffness.SizeSquared() > 0.0)
+	{
+		return true;
+	}
+	if (Settings.bLinearVelocityDriveEnabled.Max() && Settings.LinearDriveDamping.SizeSquared() > 0.0)
+	{
+		return true;
+	}
+	return false;
+}
+
+//======================================================================================================================
+void FRigPhysicsSolver::Draw(
+	FRigVMDrawInterface*                    DI,
+	const FRigPhysicsSolverSettings&        SolverSettings,
+	const FRigPhysicsVisualizationSettings& VisualizationSettings,
+	const UWorld*                           DebugWorld) const
+{
+	DECLARE_SCOPE_HIERARCHICAL_COUNTER_FUNC()
+	QUICK_SCOPE_CYCLE_COUNTER(STAT_RigPhysics_Draw);
+
+	// Lock against access to the simulation, in case of access by the WorldObject task
+	UE::TScopeLock<FTransactionallySafeCriticalSection> Lock(*SimulationMutex.Get());
+
+#if CHAOS_DEBUG_DRAW
+	// Danny TODO enable/disable chaos debug draw with a cvar or flag
+	if (DebugWorld && DebugWorld->GetPhysicsScene() && DebugWorld->GetPhysicsScene()->GetDebugDrawScene())
+	{
+		Simulation->SetDebugDrawScene(FString("ControlRig"), DebugWorld->GetPhysicsScene()->GetDebugDrawScene());
+		Simulation->DebugDraw();
+	}
+#endif
+
+	if (!DI)
+	{
+		return;
+	}
+
+	// All rendering is done relative to the component, so convert the sim space (identity) into the
+	// component space.
+	const FTransform SpaceTransform = ConvertSimSpaceTransformToComponentSpace(SolverSettings, FTransform());
+
+	if (VisualizationSettings.bShowBodies)
+	{
+		if (CollisionActorHandle)
+		{
+			DrawActor(*DI, VisualizationSettings, SpaceTransform, *CollisionActorHandle);
+		}
+
+		for (const TPair<FRigComponentKey, FRigBodyRecord>& BodyRecordPair : BodyRecords)
+		{
+			const FRigBodyRecord& Record = BodyRecordPair.Value;
+			if (const ImmediatePhysics::FActorHandle* ActorHandle = Record.ActorHandle)
+			{
+				DrawActor(*DI, VisualizationSettings, SpaceTransform, *ActorHandle);
+			}
+		}
+	}
+
+	if (VisualizationSettings.bShowJoints)
+	{
+		for (const TPair<FRigComponentKey, FRigJointRecord>& JointRecordPair : JointRecords)
+		{
+			const FRigJointRecord& Record = JointRecordPair.Value;
+			if (const ImmediatePhysics::FJointHandle* JointHandle = Record.JointHandle)
+			{
+				DrawJoint(*DI, VisualizationSettings, SpaceTransform, *JointHandle, false);
+			}
+		}
+	}
+
+	if (VisualizationSettings.bShowControls)
+	{
+		for (const TPair<FRigComponentKey, FRigControlRecord>& ControlRecordPair : ControlRecords)
+		{
+			const FRigControlRecord& Record = ControlRecordPair.Value;
+			if (const ImmediatePhysics::FJointHandle* JointHandle = Record.JointHandle)
+			{
+				if (const Chaos::FPBDJointConstraintHandle* Constraint = JointHandle->GetConstraint())
+				{
+					if (IsConstraintDriveActive(*Constraint))
+					{
+						DrawJoint(*DI, VisualizationSettings, SpaceTransform, *JointHandle, true);
+					}
+				}
+			}
+		}
+	}
+
+	if (VisualizationSettings.bShowWorldOverlapBox)
+	{
+		if (WorldOverlapBox.IsValid)
+		{
+			// Remember that rendering is done in component space
+			FTransform BoxRelWorldTM(FQuat::Identity, WorldOverlapBox.GetCenter());
+			FTransform BoxRelComponentTM = BoxRelWorldTM.GetRelativeTransform(SimulationSpaceState.ComponentTM);
+			DI->DrawBox(
+				BoxRelComponentTM,
+				FTransform(FQuat::Identity, FVector::ZeroVector, WorldOverlapBox.GetSize()),
+				FColor::Green, 
+				VisualizationSettings.LineThickness);
+		}
+	}
+
+	if (VisualizationSettings.bShowWorldObjects)
+	{
+		for (TPair<uint32, FWorldObject>& WorldObjectPair : *WorldObjects.Get())
+		{
+			FWorldObject& WorldObject = WorldObjectPair.Value;
+			if (WorldObject.ActorHandle)
+			{
+				DrawActor(*DI, VisualizationSettings, SpaceTransform, *WorldObject.ActorHandle);
+			}
+		}
+	}
+
+	int32 ShowActiveOverride = CVarControlRigPhysicsShowActiveContactsOverride.GetValueOnAnyThread();
+	int32 ShowInactiveOverride = CVarControlRigPhysicsShowInactiveContactsOverride.GetValueOnAnyThread();
+
+	bool bShowActiveContacts = ShowActiveOverride < 0 ? VisualizationSettings.bShowActiveContacts : 
+		(ShowActiveOverride ? true : false);
+	bool bShowInactiveContacts = ShowInactiveOverride < 0 ? VisualizationSettings.bShowInactiveContacts :
+		(ShowInactiveOverride ? true : false);
+
+	if (bShowActiveContacts || bShowInactiveContacts)
+	{
+		Simulation->VisitCollisions(
+			[DI, &SpaceTransform, &VisualizationSettings, bShowActiveContacts, bShowInactiveContacts]
+			(const ImmediatePhysics::FSimulation::FCollisionData& Collision)
+			{
+				using namespace Chaos;
+
+				bool bIsActive = Collision.GetCollisionAccumulatedImpulse().IsZero();
+				if ((bIsActive && bShowActiveContacts) || (!bIsActive && bShowInactiveContacts))
+				{
+					float Size = 5.0f * VisualizationSettings.ShapeSize;
+
+					int32 NumManifoldPoints = Collision.GetNumManifoldPoints();
+					for (int32 PointIndex = 0; PointIndex != NumManifoldPoints; ++PointIndex)
+					{
+						Chaos::FRealSingle Depth;
+						Chaos::FVec3 PlaneNormal;
+						Chaos::FVec3 PointLocation;
+						Chaos::FVec3 PlaneLocation;
+						Collision.GetManifoldPointData(
+							PointIndex, Depth, PlaneNormal, PointLocation, PlaneLocation);
+
+						const FVec3 PointPlaneLocation = PointLocation -
+							FVec3::DotProduct(PointLocation - PlaneLocation, PlaneNormal) * PlaneNormal;
+
+						FMatrix Axes = FRotationMatrix::MakeFromZ(PlaneNormal);
+						FTransform TM(Axes);
+						TM.SetTranslation(PointPlaneLocation);
+
+						DI->DrawCircle(
+							SpaceTransform, TM, Size,
+							bIsActive ? RigPhysicsSolverDraw::ActiveContactColor : RigPhysicsSolverDraw::InactiveContactColor,
+							VisualizationSettings.LineThickness, VisualizationSettings.ShapeDetail);
+					}
+				}
+			}, Chaos::ECollisionVisitorFlags::VisitDefault);
+	}
+}
+
diff --git a/Engine/Plugins/Experimental/ControlRigPhysics/Source/ControlRigPhysics/Private/RigPhysicsSolver_WorldObjects.cpp b/Engine/Plugins/Experimental/ControlRigPhysics/Source/ControlRigPhysics/Private/RigPhysicsSolver_WorldObjects.cpp
new file mode 100644
--- /dev/null
+++ b/Engine/Plugins/Experimental/ControlRigPhysics/Source/ControlRigPhysics/Private/RigPhysicsSolver_WorldObjects.cpp
@@ -0,0 +1,312 @@
+// Copyright Epic Games, Inc. All Rights Reserved.
+
+#include "RigPhysicsSolver.h"
+#include "RigPhysicsSolver_Space.inl"
+
+#include "Physics/ImmediatePhysics/ImmediatePhysicsActorHandle.h"
+#include "Physics/ImmediatePhysics/ImmediatePhysicsAdapters.h"
+#include "Physics/ImmediatePhysics/ImmediatePhysicsSimulation.h"
+#include "Physics/ImmediatePhysics/ImmediatePhysicsChaos/ImmediatePhysicsSimulation_Chaos.h"
+#include "Chaos/ParticleHandle.h"
+
+#include "Engine/OverlapResult.h"
+#include "Engine/World.h"
+#include "CollisionQueryParams.h"
+#include "CollisionShape.h"
+
+#include "Stats/Stats.h"
+#include "Stats/StatsHierarchical.h"
+
+#include "Components/InstancedStaticMeshComponent.h"
+
+//======================================================================================================================
+// Note that the simulation mutex will already be taken when we are called
+void FRigPhysicsSolver::UpdateWorldObjectsPrePhysics(const FRigPhysicsSolverSettings& SolverSettings)
+{
+	DECLARE_SCOPE_HIERARCHICAL_COUNTER_FUNC()
+	QUICK_SCOPE_CYCLE_COUNTER(STAT_RigPhysics_UpdateWorldObjectsPrePhysics);
+
+	if (SolverSettings.WorldCollisionType == ERigPhysicsWorldCollisionType::None)
+	{
+		for (TPair<uint32, FWorldObject>& WorldObjectPair : *WorldObjects.Get())
+		{
+			FWorldObject& WorldObject = WorldObjectPair.Value;
+			if (WorldObject.ActorHandle)
+			{
+				UE_LOG(LogRigPhysics, Warning, TEXT("Expiring world object %s"),
+					*WorldObject.ActorHandle->GetName().ToString());
+
+				Simulation->DestroyActor(WorldObject.ActorHandle);
+			}
+		}
+		WorldObjects->Reset();
+	}
+	else
+	{
+		TArray<uint32> EntriesToRemove;
+		for (TPair<uint32, FWorldObject>& WorldObjectPair : *WorldObjects.Get())
+		{
+			FWorldObject& WorldObject = WorldObjectPair.Value;
+			if (WorldObject.GetExpired(UpdateCounter, SolverSettings.WorldCollisionExpiryFrames))
+			{
+				if (WorldObject.ActorHandle)
+				{
+					UE_LOG(LogRigPhysics, Log, TEXT("Expiring world object %s"), 
+						*WorldObject.ActorHandle->GetName().ToString());
+
+					Simulation->DestroyActor(WorldObject.ActorHandle);
+					WorldObject.ActorHandle = nullptr;
+				}
+				EntriesToRemove.Add(WorldObjectPair.Key);
+			}
+		}
+		for (uint32 EntryToRemove : EntriesToRemove)
+		{
+			WorldObjects->Remove(EntryToRemove);
+		}
+	}
+
+	for (TPair<uint32, FWorldObject>& WorldObjectPair : *WorldObjects.Get())
+	{
+		FWorldObject& WorldObject = WorldObjectPair.Value;
+		FTransform SimSpaceTM = ConvertWorldTransformToSimSpace(SolverSettings, WorldObject.ComponentWorldTransform);
+		// Note that the code above should have ensured we don't have a null pointer, but let's be safe
+		if (ensure(WorldObject.ActorHandle != nullptr))
+		{
+			if (WorldObject.LastSeenUpdateCounter == -1)
+			{
+				UE_LOG(LogRigPhysics, Log,
+					TEXT("Adding object %s"), *WorldObject.ActorHandle->GetName().ToString());
+
+				// We need to call InitWorldTransform (not just SetWorldTransform) here, to ensure the
+				// velocities are set to zero (though they should be after creation anyway), and that
+				// the object has its bounds set.
+				WorldObject.ActorHandle->InitWorldTransform(SimSpaceTM);
+				// If the object is found in the next overlap, then its update counter will be updated.
+				// However, in case it's not, we set it here too to make sure that it expires. 
+				WorldObject.LastSeenUpdateCounter = UpdateCounter;
+			}
+			else
+			{
+				WorldObject.ActorHandle->SetKinematicTarget(SimSpaceTM);
+			}
+		}
+	}
+}
+
+//======================================================================================================================
+// Note that the simulation mutex will already be taken/locked when we are called 
+// from FRigPhysicsSolver::StepSimulation
+void FRigPhysicsSolver::UpdateWorldObjectsPostPhysics(
+	const UWorld* World, const FRigPhysicsSolverSettings& SolverSettings, const AActor* OwningActorPtr)
+{
+	DECLARE_SCOPE_HIERARCHICAL_COUNTER_FUNC()
+	QUICK_SCOPE_CYCLE_COUNTER(STAT_RigPhysics_UpdateWorldObjectsPostPhysics);
+
+	WorldOverlapBox = FBox();
+	if (SolverSettings.WorldCollisionType == ERigPhysicsWorldCollisionType::None)
+	{
+		return;
+	}
+
+	// Get the volume over which to do overlaps
+	const FTransform SpaceTransform = 
+		ConvertSimSpaceTransformToComponentSpace(SolverSettings, FTransform()) * SimulationSpaceState.ComponentTM;
+	for (TPair<FRigComponentKey, FRigBodyRecord>& BodyRecordPair : BodyRecords)
+	{
+		const FRigComponentKey& ComponentKey = BodyRecordPair.Key;
+		const FRigBodyRecord& Record = BodyRecordPair.Value;
+
+		if (const ImmediatePhysics::FActorHandle* ActorHandle = Record.ActorHandle)
+		{
+			if (!ActorHandle->GetIsKinematic())
+			{
+				const Chaos::FGeometryParticleHandle* GeometryParticleHandle = ActorHandle->GetParticle();
+				const FTransform ParticleTransform = GeometryParticleHandle->GetTransformXR() * SpaceTransform;
+
+				const auto& AABB = GeometryParticleHandle->LocalBounds();
+				if (!AABB.IsEmpty())
+				{
+					for (int32 Index = 0; Index != 8; ++Index)
+					{
+						WorldOverlapBox += ParticleTransform.TransformPosition(AABB.GetVertex(Index));
+					}
+				}
+			}
+		}
+	}
+
+	if (!WorldOverlapBox.IsValid)
+	{
+		return;
+	}
+
+	TWeakPtr<TMap<uint32, FWorldObject>> WeakWorldObjects = WorldObjects;
+	int64 UpdateCounterCopy = UpdateCounter;
+	const AActor* OwningActorPtrCopy = OwningActorPtr;
+	TSharedPtr<ImmediatePhysics::FSimulation> SimulationCopy = Simulation;
+
+	const TWeakObjectPtr<const UWorld> WeakWorld = World;
+
+	FVector NewHalfExtent = WorldOverlapBox.GetExtent() * SolverSettings.WorldCollisionBoundsExpansion;
+	FVector OverlapAABoxCenter = WorldOverlapBox.GetCenter();
+	WorldOverlapBox = FBox(OverlapAABoxCenter - NewHalfExtent, OverlapAABoxCenter + NewHalfExtent);
+	FCollisionShape OverlapAABox = FCollisionShape::MakeBox(NewHalfExtent);
+
+	FCollisionQueryParams QueryParams = FCollisionQueryParams(
+		SCENE_QUERY_STAT(ControlRigPhysicsFindGeometry), /*bTraceComplex=*/false);
+	QueryParams.MobilityType = EQueryMobilityType::Any;
+
+	FCollisionObjectQueryParams ObjectQueryParams;
+	switch (SolverSettings.WorldCollisionType)
+	{
+	case ERigPhysicsWorldCollisionType::Static:
+		ObjectQueryParams = FCollisionObjectQueryParams(FCollisionObjectQueryParams::InitType::AllStaticObjects); break;
+	case ERigPhysicsWorldCollisionType::Dynamic:
+		ObjectQueryParams = FCollisionObjectQueryParams(FCollisionObjectQueryParams::InitType::AllDynamicObjects); break;
+	case ERigPhysicsWorldCollisionType::All:
+		ObjectQueryParams = FCollisionObjectQueryParams(FCollisionObjectQueryParams::InitType::AllObjects); break;
+	default:
+		break;
+	}
+
+	// The actual overlap needs to be on the game thread, so we will run it in an async task.
+	// However, note that between now and then, we/our simulation etc might get destroyed so we must
+	// protect against that. Even the mutex guarding against access to the simulation may be released.
+	if (ObjectQueryParams.GetObjectTypesToQuery())
+	{
+		// Make a weak copy of the simulation mutex and pass it to the (game thread) task. Note that
+		// when we come to use it we need to check whether it it valid.
+		TWeakPtr<FTransactionallySafeCriticalSection> WeakSimulationMutex = SimulationMutex;
+
+		AsyncTask(ENamedThreads::GameThread,
+			[WeakSimulationMutex, WeakWorld, WeakWorldObjects, OverlapAABox, OverlapAABoxCenter, QueryParams, 
+			ObjectQueryParams, UpdateCounterCopy, OwningActorPtrCopy, SimulationCopy]()
+			{
+				if (TSharedPtr<FTransactionallySafeCriticalSection> SimulationMutex = WeakSimulationMutex.Pin())
+				{
+					UE::TScopeLock<FTransactionallySafeCriticalSection> Lock(*SimulationMutex.Get());
+
+					if (TStrongObjectPtr<const UWorld> World = WeakWorld.Pin())
+					{
+						if (TSharedPtr<TMap<uint32, FWorldObject>> WorldObjects = WeakWorldObjects.Pin())
+						{
+							TArray<FOverlapResult> Overlaps;
+							World->OverlapMultiByObjectType(
+								Overlaps, OverlapAABoxCenter, FQuat::Identity, ObjectQueryParams,
+								OverlapAABox, QueryParams);
+
+							for (const FOverlapResult& Overlap : Overlaps)
+							{
+								if (UPrimitiveComponent* OverlapComp = Overlap.GetComponent())
+								{
+									// For instanced static mesh components, each overlap result
+									// corresponds to a specific instance via Overlap.ItemIndex.
+									// The component-level BodyInstance has null BodySetup and no
+									// physics actor, so we must use the per-instance body instead.
+									int32 OverlapInstanceIndex = INDEX_NONE;
+									if (OverlapComp->IsA<UInstancedStaticMeshComponent>())
+									{
+										OverlapInstanceIndex = Overlap.ItemIndex;
+									}
+
+									// Use a per-instance composite key for ISM so each instance
+									// gets its own FWorldObject entry.
+									const uint32 Key = (OverlapInstanceIndex == INDEX_NONE)
+										? OverlapComp->GetUniqueID()
+										: HashCombine(OverlapComp->GetUniqueID(), (uint32)OverlapInstanceIndex);
+
+									FWorldObject* WorldObject = WorldObjects->Find(Key);
+									if (!WorldObject)
+									{
+										// New object - add it to the sim
+										if (OverlapComp->GetOwner() && OverlapComp->GetOwner() != OwningActorPtrCopy)
+										{
+											// Resolve the correct BodyInstance and transform.
+											// For ISM the component-level BodyInstance has null
+											// BodySetup, so use the per-instance body.
+											const FBodyInstance* BodyInstanceToUse = &OverlapComp->BodyInstance;
+
+											if (OverlapInstanceIndex != INDEX_NONE)
+											{
+												UInstancedStaticMeshComponent* ISMComp = Cast<UInstancedStaticMeshComponent>(OverlapComp);
+												BodyInstanceToUse = ISMComp	
+													? ISMComp->GetBodyInstance(NAME_None, false, OverlapInstanceIndex)
+													: nullptr;
+												if (BodyInstanceToUse == nullptr)
+												{
+													// Instance has no physics body
+													continue;
+												}
+											}
+
+											// Chaos holds scale separately to the world transform of the body
+											FTransform TransformToUse = BodyInstanceToUse->GetUnrealWorldTransform();
+											TransformToUse.SetScale3D(BodyInstanceToUse->Scale3D);
+
+											WorldObject = &WorldObjects->Add(Key, FWorldObject());
+
+											WorldObject->WorldPrimitiveComponent = OverlapComp;
+											WorldObject->InstanceBodyIndex = OverlapInstanceIndex;
+
+											// Note that we need to create the simulation actor here, even
+											// though we're in the game thread, so we can copy data out of the
+											// object. This should be safe because we only run this task after
+											// we've finished using the simulation in the worker thread.
+
+											WorldObject->ActorHandle = SimulationCopy->CreateActor(
+												ImmediatePhysics::MakeKinematicActorSetup(
+													BodyInstanceToUse, TransformToUse));
+
+											if (WorldObject->ActorHandle)
+											{
+#if WITH_EDITOR
+												WorldObject->ActorHandle->SetName(*OverlapComp->GetOwner()->GetActorLabel());
+#endif
+												SimulationCopy->AddToCollidingPairs(WorldObject->ActorHandle);
+											}
+											WorldObject->ComponentWorldTransform = TransformToUse;
+
+											// Flag that the simulation TM needs to be set, rather
+											// than using a kinematic target.
+											WorldObject->LastSeenUpdateCounter = -1;
+										}
+									}
+									else
+									{
+										if (TStrongObjectPtr<UPrimitiveComponent> WPC =
+											WorldObject->WorldPrimitiveComponent.Pin())
+										{
+											FBodyInstance* BodyInstanceToUse = &WPC->BodyInstance;
+
+											if (WorldObject->InstanceBodyIndex != INDEX_NONE)
+											{
+												UInstancedStaticMeshComponent* ISMComp = Cast<UInstancedStaticMeshComponent>(WPC.Get());
+												BodyInstanceToUse = ISMComp 
+													? ISMComp->GetBodyInstance(NAME_None, false, WorldObject->InstanceBodyIndex)
+													: nullptr;
+
+												if (BodyInstanceToUse == nullptr)
+												{
+													// Instance has no physics body
+													continue;
+												}
+											}
+
+											// Chaos holds scale separately to the world transform
+											FTransform TransformToUse = BodyInstanceToUse->GetUnrealWorldTransform();
+											TransformToUse.SetScale3D(BodyInstanceToUse->Scale3D);
+
+											WorldObject->ComponentWorldTransform = TransformToUse;
+											WorldObject->LastSeenUpdateCounter = UpdateCounterCopy;
+										}
+									}
+								}
+							}
+						}
+					}
+				}
+			});
+	}
+}
+
diff --git a/Engine/Plugins/Experimental/ControlRigPhysics/Source/ControlRigPhysics/Public/RigPhysicsData.h b/Engine/Plugins/Experimental/ControlRigPhysics/Source/ControlRigPhysics/Public/RigPhysicsData.h
--- a/Engine/Plugins/Experimental/ControlRigPhysics/Source/ControlRigPhysics/Public/RigPhysicsData.h
+++ b/Engine/Plugins/Experimental/ControlRigPhysics/Source/ControlRigPhysics/Public/RigPhysicsData.h
@@ -4,8 +4,8 @@
 
 #include "PhysicsControlData.h"
 
-#include "Rigs/RigPhysics.h"
 #include "Rigs/RigHierarchyDefines.h"
+#include "Rigs/RigHierarchyComponents.h"
 
 #include "PhysicsEngine/ConstraintTypes.h"
 #include "PhysicsEngine/ConstraintDrives.h"
diff --git a/Engine/Plugins/Experimental/ControlRigPhysics/Source/ControlRigPhysics/Public/RigPhysicsExecution.h b/Engine/Plugins/Experimental/ControlRigPhysics/Source/ControlRigPhysics/Public/RigPhysicsExecution.h
--- a/Engine/Plugins/Experimental/ControlRigPhysics/Source/ControlRigPhysics/Public/RigPhysicsExecution.h
+++ b/Engine/Plugins/Experimental/ControlRigPhysics/Source/ControlRigPhysics/Public/RigPhysicsExecution.h
@@ -30,31 +30,6 @@ struct FRigUnit_PhysicsBase : public FRigUnit
 	GENERATED_BODY()
 };
 
-/**
- * Instantiates all the objects in the physics world. Some properties can't be modified after this happens.
- * Note that it will happen automatically during the first simulation step if it hasn't been explicitly
- * requested. Explicit instantiation allows the timing to be controlled, as allocations etc may cause some
- * delays.
- */
-USTRUCT(meta = (DisplayName = "Instantiate physics"))
-struct FRigUnit_InstantiatePhysics : public FRigUnit_PhysicsBaseMutable
-{
-	GENERATED_BODY()
-
-	FRigUnit_InstantiatePhysics()
-	{
-		PhysicsSolverComponentKey.ElementKey.Type = ERigElementType::Bone;
-		PhysicsSolverComponentKey.Name = FRigPhysicsSolverComponent::GetDefaultName();
-	}
-
-	RIGVM_METHOD()
-	UE_API virtual void Execute() override;
-
-	// The solver that should be instantiated
-	UPROPERTY(meta = (Input))
-	FRigComponentKey PhysicsSolverComponentKey;
-};
-
 // Adds a set of physics components including the body, joint and controls
 USTRUCT(meta = (DisplayName = "Spawn Physics Components", Keywords = "Add,Construction,Create,New,Body,Joint,Control", Varying))
 struct FRigUnit_AddPhysicsComponents : public FRigUnit_PhysicsBaseMutable
diff --git a/Engine/Plugins/Experimental/ControlRigPhysics/Source/ControlRigPhysics/Public/RigPhysicsSolverComponent.h b/Engine/Plugins/Experimental/ControlRigPhysics/Source/ControlRigPhysics/Public/RigPhysicsSolverComponent.h
--- a/Engine/Plugins/Experimental/ControlRigPhysics/Source/ControlRigPhysics/Public/RigPhysicsSolverComponent.h
+++ b/Engine/Plugins/Experimental/ControlRigPhysics/Source/ControlRigPhysics/Public/RigPhysicsSolverComponent.h
@@ -10,7 +10,11 @@
 
 #define UE_API CONTROLRIGPHYSICS_API
 
+struct FRigPhysicsSolver;
+
+//======================================================================================================================
 // A solver coordinates the physical movement of bodies
+//======================================================================================================================
 USTRUCT(BlueprintType)
 struct FRigPhysicsSolverComponent : public FRigBaseComponent
 {
@@ -40,9 +44,24 @@ public:
 	UE_API virtual void OnAddedToHierarchy(URigHierarchy* InHierarchy, URigHierarchyController* InController) override;
 
 	virtual FName GetDefaultComponentName() const override { return GetDefaultName(); }
-
 	static FName GetDefaultName() { return TEXT("PhysicsSolver"); }
 
+	// This will make the internally owned solver if necessary, and return it. Note that there's no
+	// danger of this being called concurrently from different threads, since each FRigPhysicsSolver
+	// is owned by just one rig.
+	FRigPhysicsSolver* GetPhysicsSolver() const;
+
+private:
+	// We use a TSharedPtr here to share access/ownership with the game-thread world collisions task.
+	//
+	// We make it mutable so that GetPhysicsSolver can be const - in practice the user shouldn't
+	// have to worry.
+	//
+	// Finally, UE's reflection system wants to be able to use the assignment operator. If it wasn't
+	// for that, we could just delete the copy operator etc and prevent duplicates containing a
+	// pointer to the same PhysicsSolver. In practice, this won't ever happen for a "live"
+	// component, so we're OK.
+	mutable TSharedPtr<FRigPhysicsSolver> PhysicsSolverPtr;
 };
 
 #undef UE_API
diff --git a/Engine/Plugins/Experimental/ControlRigPhysics/Source/ControlRigPhysics/Public/RigPhysicsSolverExecution.h b/Engine/Plugins/Experimental/ControlRigPhysics/Source/ControlRigPhysics/Public/RigPhysicsSolverExecution.h
--- a/Engine/Plugins/Experimental/ControlRigPhysics/Source/ControlRigPhysics/Public/RigPhysicsSolverExecution.h
+++ b/Engine/Plugins/Experimental/ControlRigPhysics/Source/ControlRigPhysics/Public/RigPhysicsSolverExecution.h
@@ -51,6 +51,31 @@ struct FRigUnit_AddPhysicsSolver : public FRigUnit_PhysicsBaseMutable
 	FRigPhysicsSimulationSpaceSettings SimulationSpaceSettings;
 };
 
+/**
+* Instantiates all the objects in the physics world. Some properties can't be modified after this happens.
+* Note that it will happen automatically during the first simulation step if it hasn't been explicitly
+* requested. Explicit instantiation allows the timing to be controlled, as allocations etc may cause some
+* delays.
+*/
+USTRUCT(meta = (DisplayName = "Instantiate physics"))
+struct FRigUnit_InstantiatePhysics : public FRigUnit_PhysicsBaseMutable
+{
+	GENERATED_BODY()
+
+	FRigUnit_InstantiatePhysics()
+	{
+		PhysicsSolverComponentKey.ElementKey.Type = ERigElementType::Bone;
+		PhysicsSolverComponentKey.Name = FRigPhysicsSolverComponent::GetDefaultName();
+	}
+
+	RIGVM_METHOD()
+	UE_API virtual void Execute() override;
+
+	// The solver that should be instantiated
+	UPROPERTY(meta = (Input))
+	FRigComponentKey PhysicsSolverComponentKey;
+};
+
 /**
  * Steps the specified physics solver
  */
