Enable Control Rig Physics plugin by default

#Synced-CL 52251316

diff --git a/Engine/Plugins/Experimental/ControlRigPhysics/ControlRigPhysics.uplugin b/Engine/Plugins/Experimental/ControlRigPhysics/ControlRigPhysics.uplugin
--- a/Engine/Plugins/Experimental/ControlRigPhysics/ControlRigPhysics.uplugin
+++ b/Engine/Plugins/Experimental/ControlRigPhysics/ControlRigPhysics.uplugin
@@ -12,8 +12,9 @@
 	"SupportURL": "",
 	"CanContainContent": true,
 	"CanContainVerse": false,
-	"IsBetaVersion": false,
-	"IsExperimentalVersion": true,
+	"IsBetaVersion": true,
+	"IsExperimentalVersion": false,
+	"EnabledByDefault": true,
 	"Installed": false,
 	"Plugins": [
 		{
@@ -23,13 +24,13 @@
 		{
 			"Name": "PhysicsControl",
 			"Enabled": true
-		},
+		}
 	],
 	"Modules": [
 		{
 			"Name": "ControlRigPhysics",
 			"Type": "Runtime",
 			"LoadingPhase": "Default"
-		},
+		}
 	]
 }
\ No newline at end of file
