Control rig physics - disable visualisation by default, and minor fixups to cvars and help text

#Synced-CL 52541170

diff --git a/Engine/Plugins/Experimental/ControlRigPhysics/Source/ControlRigPhysics/Private/RigPhysicsSolverExecution.cpp b/Engine/Plugins/Experimental/ControlRigPhysics/Source/ControlRigPhysics/Private/RigPhysicsSolverExecution.cpp
--- a/Engine/Plugins/Experimental/ControlRigPhysics/Source/ControlRigPhysics/Private/RigPhysicsSolverExecution.cpp
+++ b/Engine/Plugins/Experimental/ControlRigPhysics/Source/ControlRigPhysics/Private/RigPhysicsSolverExecution.cpp
@@ -13,8 +13,8 @@
 #include UE_INLINE_GENERATED_CPP_BY_NAME(RigPhysicsSolverExecution)
 
 TAutoConsoleVariable<int> CVarControlRigPhysicsVisualizationOverride(
-	TEXT("ControlRig.Physics.VisualizationOveride"), -1,
-	TEXT("-1 uses the Step Physics Solver setting, 0 forces visualization to be disabled, 1 forces it to be enabled."));
+	TEXT("ControlRig.Physics.VisualizationOverride"), -1,
+	TEXT("-1 uses the visualization settings from Step Physics Solver, 0 forces visualization to be disabled, 1 forces it to be enabled."));
 
 TAutoConsoleVariable<int> CVarControlRigPhysicsEnableStepSolver(
 	TEXT("ControlRig.Physics.EnableStepSolver"), 1,
diff --git a/Engine/Plugins/Experimental/ControlRigPhysics/Source/ControlRigPhysics/Public/RigPhysicsData.h b/Engine/Plugins/Experimental/ControlRigPhysics/Source/ControlRigPhysics/Public/RigPhysicsData.h
--- a/Engine/Plugins/Experimental/ControlRigPhysics/Source/ControlRigPhysics/Public/RigPhysicsData.h
+++ b/Engine/Plugins/Experimental/ControlRigPhysics/Source/ControlRigPhysics/Public/RigPhysicsData.h
@@ -31,10 +31,11 @@ struct FRigPhysicsVisualizationSettings
 {
 	GENERATED_BODY()
 
-	// Enable visualization (includes enabling low-level chaos debug draw). Note that this can be
-	// overridden using CVar ControlRig.Physics.VisualizationOverride
+	// Enable visualization (includes enabling low-level chaos debug draw). There is a significant
+	// runtime cost to enabling visualization, so make sure you disable it when it's not needed, and
+	// when you are profiling. Note that this can be overridden using CVar ControlRig.Physics.VisualizationOverride
 	UPROPERTY(EditAnywhere, Category = Visualization)
-	bool bEnableVisualization = true;
+	bool bEnableVisualization = false;
 
 	UPROPERTY(EditAnywhere, Category = Visualization, meta = (ClampMin = "0.0"))
 	float LineThickness = 1.0f;
